fix for unacceptable words check.

#rnx

diff --git a/Engine/Plugins/Experimental/DefaultInstallBundleManager/Source/Private/InstallBundleManagerUtil.cpp b/Engine/Plugins/Experimental/DefaultInstallBundleManager/Source/Private/InstallBundleManagerUtil.cpp
--- a/Engine/Plugins/Experimental/DefaultInstallBundleManager/Source/Private/InstallBundleManagerUtil.cpp
+++ b/Engine/Plugins/Experimental/DefaultInstallBundleManager/Source/Private/InstallBundleManagerUtil.cpp
@@ -261,7 +261,7 @@ namespace InstallBundleManagerUtil
 				//Before cleaning up, lets save out all dirty stats to make sure we have the most up-to-date stuff on disk.
 				SaveAllDirtyStatsToDisk();
 			
-				//Nuke data for our bundles
+				//remove data for our bundles
 				TArray<FString> RequiredBundlesForSession;
 				FoundSessionStats->GetRequiredBundles(RequiredBundlesForSession);
 				for (const FString& BundleName : RequiredBundlesForSession)
