Cleaned up DefaultInstallBundleManager to not set a private pch file(it isn't force included) and to enable IWYU.

diff --git a/Engine/Plugins/Experimental/DefaultInstallBundleManager/Source/DefaultInstallBundleManager.Build.cs b/Engine/Plugins/Experimental/DefaultInstallBundleManager/Source/DefaultInstallBundleManager.Build.cs
--- a/Engine/Plugins/Experimental/DefaultInstallBundleManager/Source/DefaultInstallBundleManager.Build.cs
+++ b/Engine/Plugins/Experimental/DefaultInstallBundleManager/Source/DefaultInstallBundleManager.Build.cs
@@ -8,9 +8,6 @@ public class DefaultInstallBundleManager : ModuleRules
 {
 	public DefaultInstallBundleManager(ReadOnlyTargetRules Target) : base(Target)
 	{
-		PCHUsage = PCHUsageMode.UseSharedPCHs;
-		PrivatePCHHeaderFile = "Private/DefaultInstallBundleManagerPrivatePCH.h";
-
 		PublicDependencyModuleNames.AddRange(
 			new string[] {
 				"InstallBundleManager",
diff --git a/Engine/Plugins/Experimental/DefaultInstallBundleManager/Source/Private/DefaultInstallBundleManager.cpp b/Engine/Plugins/Experimental/DefaultInstallBundleManager/Source/Private/DefaultInstallBundleManager.cpp
--- a/Engine/Plugins/Experimental/DefaultInstallBundleManager/Source/Private/DefaultInstallBundleManager.cpp
+++ b/Engine/Plugins/Experimental/DefaultInstallBundleManager/Source/Private/DefaultInstallBundleManager.cpp
@@ -1,8 +1,8 @@
 // Copyright Epic Games, Inc. All Rights Reserved.
 
-#include "DefaultInstallBundleManagerPrivatePCH.h"
-
 #include "DefaultInstallBundleManager.h"
+
+#include "DefaultInstallBundleManagerPrivate.h"
 #include "InstallBundleManagerUtil.h"
 #include "InstallBundleManagerInterface.h"
 #include "InstallBundleSourceInterface.h"
diff --git a/Engine/Plugins/Experimental/DefaultInstallBundleManager/Source/Private/DefaultInstallBundleManagerPrivate.h b/Engine/Plugins/Experimental/DefaultInstallBundleManager/Source/Private/DefaultInstallBundleManagerPrivate.h
new file mode 100644
--- /dev/null
+++ b/Engine/Plugins/Experimental/DefaultInstallBundleManager/Source/Private/DefaultInstallBundleManagerPrivate.h
@@ -0,0 +1,45 @@
+// Copyright Epic Games, Inc. All Rights Reserved.
+
+#pragma once
+
+#include "CoreMinimal.h"
+#include "ProfilingDebugging/CsvProfiler.h"
+
+CSV_DECLARE_CATEGORY_MODULE_EXTERN(INSTALLBUNDLEMANAGER_API, InstallBundleManager);
+
+#define LOG_INSTALL_BUNDLE_MAN(Verbosity, Format, ...) UE_LOG(LogDefaultInstallBundleManager, Verbosity, Format, ##__VA_ARGS__)
+
+#define LOG_INSTALL_BUNDLE_MAN_OVERRIDE(InVerbosityOverride, Verbosity, Format, ...) \
+{ \
+	if((ELogVerbosity::Verbosity & ELogVerbosity::VerbosityMask) < ELogVerbosity::Display) \
+	{ \
+		UE_LOG(LogDefaultInstallBundleManager, Verbosity, Format, ##__VA_ARGS__) \
+	} \
+	else \
+	{ \
+		ELogVerbosity::Type VerbosityOverride = (ELogVerbosity::Type)(InVerbosityOverride & ELogVerbosity::VerbosityMask); \
+		if (VerbosityOverride <= ELogVerbosity::Verbosity) \
+		{ \
+			UE_LOG(LogDefaultInstallBundleManager, Verbosity, Format, ##__VA_ARGS__) \
+		} \
+		else \
+		{ \
+			switch (VerbosityOverride) \
+			{ \
+			case ELogVerbosity::Log: \
+				UE_LOG(LogDefaultInstallBundleManager, Log, Format, ##__VA_ARGS__) \
+				break; \
+			case ELogVerbosity::Verbose: \
+				UE_LOG(LogDefaultInstallBundleManager, Verbose, Format, ##__VA_ARGS__) \
+				break; \
+			case ELogVerbosity::VeryVerbose: \
+				UE_LOG(LogDefaultInstallBundleManager, VeryVerbose, Format, ##__VA_ARGS__) \
+				break; \
+			default: \
+				checkf(false, TEXT("Unkown verbosity override")); \
+				UE_LOG(LogDefaultInstallBundleManager, Verbosity, Format, ##__VA_ARGS__) \
+				break; \
+			} \
+		} \
+	} \
+}
diff --git a/Engine/Plugins/Experimental/DefaultInstallBundleManager/Source/Private/DefaultInstallBundleManagerPrivatePCH.h b/Engine/Plugins/Experimental/DefaultInstallBundleManager/Source/Private/DefaultInstallBundleManagerPrivatePCH.h
deleted file mode 100644
diff --git a/Engine/Plugins/Experimental/DefaultInstallBundleManager/Source/Private/InstallBundleManagerUtil.cpp b/Engine/Plugins/Experimental/DefaultInstallBundleManager/Source/Private/InstallBundleManagerUtil.cpp
--- a/Engine/Plugins/Experimental/DefaultInstallBundleManager/Source/Private/InstallBundleManagerUtil.cpp
+++ b/Engine/Plugins/Experimental/DefaultInstallBundleManager/Source/Private/InstallBundleManagerUtil.cpp
@@ -1,9 +1,8 @@
 // Copyright Epic Games, Inc. All Rights Reserved.
 
-#include "DefaultInstallBundleManagerPrivatePCH.h"
-
 #include "InstallBundleManagerUtil.h"
 
+#include "DefaultInstallBundleManagerPrivate.h"
 #include "Misc/App.h"
 
 #include "JsonObjectConverter.h"
diff --git a/Engine/Plugins/Experimental/DefaultInstallBundleManager/Source/Private/InstallBundleSourceBulk.cpp b/Engine/Plugins/Experimental/DefaultInstallBundleManager/Source/Private/InstallBundleSourceBulk.cpp
--- a/Engine/Plugins/Experimental/DefaultInstallBundleManager/Source/Private/InstallBundleSourceBulk.cpp
+++ b/Engine/Plugins/Experimental/DefaultInstallBundleManager/Source/Private/InstallBundleSourceBulk.cpp
@@ -1,8 +1,8 @@
 // Copyright Epic Games, Inc. All Rights Reserved.
 
 #include "InstallBundleSourceBulk.h"
-#include "DefaultInstallBundleManagerPrivatePCH.h"
 
+#include "DefaultInstallBundleManagerPrivate.h"
 #include "IPlatformFilePak.h"
 #include "Containers/Ticker.h"
 #include "Misc/CommandLine.h"
