Adding very simple base class for platform install bundle sources.

#rnx

diff --git a/Engine/Plugins/Experimental/DefaultInstallBundleManager/Source/Private/InstallBundleSourcePlatformBase.cpp b/Engine/Plugins/Experimental/DefaultInstallBundleManager/Source/Private/InstallBundleSourcePlatformBase.cpp
new file mode 100644
--- /dev/null
+++ b/Engine/Plugins/Experimental/DefaultInstallBundleManager/Source/Private/InstallBundleSourcePlatformBase.cpp
@@ -0,0 +1,144 @@
+// Copyright Epic Games, Inc. All Rights Reserved.
+
+#include "InstallBundleSourcePlatformBase.h"
+#include "DefaultInstallBundleManagerPrivatePCH.h"
+
+#include "InstallBundleManagerUtil.h"
+
+#include "Misc/ConfigCacheIni.h"
+#include "Stats/Stats.h"
+
+#define LOG_SOURCE_PLATFORMBASE(Verbosity, Format, ...) LOG_INSTALL_BUNDLE_MAN(Verbosity, TEXT("InstallBundleSourcePlatformBase: ") Format, ##__VA_ARGS__)
+
+#define LOG_SOURCE_PLATFORMBASE_OVERRIDE(VerbosityOverride, Verbosity, Format, ...) LOG_INSTALL_BUNDLE_MAN_OVERRIDE(VerbosityOverride, Verbosity, TEXT("InstallBundleSourcePlatformBase: ") Format, ##__VA_ARGS__)
+
+
+FInstallBundleSourcePlatformBase::FInstallBundleSourcePlatformBase()
+{
+}
+
+FInstallBundleSourcePlatformBase::~FInstallBundleSourcePlatformBase()
+{
+}
+
+
+EInstallBundleSourceType FInstallBundleSourcePlatformBase::GetSourceType() const
+{
+	return EInstallBundleSourceType::Platform;
+}
+
+FInstallBundleSourceInitInfo FInstallBundleSourcePlatformBase::Init(
+	TSharedRef<InstallBundleUtil::FContentRequestStatsMap> InRequestStats, 
+	TSharedPtr<IAnalyticsProviderET> InAnalyticsProvider,
+	TSharedPtr<InstallBundleUtil::PersistentStats::FPersistentStatContainerBase> InPersistentStatsContainer)
+{
+	CSV_SCOPED_TIMING_STAT(InstallBundleManager, FInstallBundleSourcePlatformBase_Init);
+
+	RequestStats = MoveTemp(InRequestStats);
+	AnalyticsProvider = MoveTemp(InAnalyticsProvider);
+	PersistentStatsContainer = MoveTemp(InPersistentStatsContainer);
+
+	FInstallBundleSourceInitInfo InitInfo;
+
+	const FConfigFile* InstallBundleConfig = GConfig->FindConfigFile(GInstallBundleIni);
+	if (!InstallBundleConfig)
+	{
+		InitState = EInstallBundleManagerInitState::Failed;
+		InitInfo.Result = EInstallBundleManagerInitResult::ConfigurationError;
+
+		LOG_SOURCE_PLATFORMBASE(Error, TEXT("Initialization Failed - %s"), LexToString(InitInfo.Result));
+		return InitInfo;
+	}
+
+	for (const TPair<FString, FConfigSection>& Pair : *InstallBundleConfig)
+	{
+		const FString& Section = Pair.Key;
+		if (!Section.StartsWith(InstallBundleUtil::GetInstallBundleSectionPrefix()))
+			continue;
+
+		if (!IsPlatformInstallBundleConfig(InstallBundleConfig, Section))
+			continue;
+
+		TArray<FString> StrSearchRegexPatterns;
+		if (!InstallBundleConfig->GetArray(*Section, TEXT("FileRegex"), StrSearchRegexPatterns))
+			continue;
+
+		TArray<FRegexPattern> SearchRegexPatterns;
+		SearchRegexPatterns.Reserve(StrSearchRegexPatterns.Num());
+		for (const FString& Str : StrSearchRegexPatterns)
+		{
+			SearchRegexPatterns.Emplace(Str, ERegexPatternFlags::CaseInsensitive);
+		}
+
+		const FString BundleName = Section.RightChop(InstallBundleUtil::GetInstallBundleSectionPrefix().Len());
+		BundleRegexList.Emplace(TPair<FString, TArray<FRegexPattern>>(BundleName, MoveTemp(SearchRegexPatterns)));
+	}
+
+	BundleRegexList.StableSort([](const TPair<FString, TArray<FRegexPattern>>& PairA, const TPair<FString, TArray<FRegexPattern>>& PairB) -> bool
+	{
+		int BundleAOrder = INT_MAX;
+		int BundleBOrder = INT_MAX;
+
+		const FString SectionA = InstallBundleUtil::GetInstallBundleSectionPrefix() + PairA.Key;
+		const FString SectionB = InstallBundleUtil::GetInstallBundleSectionPrefix() + PairB.Key;
+
+		if (!GConfig->GetInt(*SectionA, TEXT("Order"), BundleAOrder, GInstallBundleIni))
+		{
+			LOG_SOURCE_PLATFORMBASE(Warning, TEXT("Bundle Section %s doesn't have an order"), *SectionA);
+		}
+
+		if (!GConfig->GetInt(*SectionB, TEXT("Order"), BundleBOrder, GInstallBundleIni))
+		{
+			LOG_SOURCE_PLATFORMBASE(Warning, TEXT("Bundle Section %s doesn't have an order"), *SectionB);
+		}
+
+		return BundleAOrder < BundleBOrder;
+	});
+
+	return InitInfo;
+}
+
+void FInstallBundleSourcePlatformBase::AsyncInit_QueryBundleInfo(FInstallBundleSourceQueryBundleInfoDelegate Callback)
+{
+	CSV_SCOPED_TIMING_STAT(InstallBundleManager, FInstallBundleSourcePlatformBase_AsyncInit_QueryBundleInfo);
+
+	check(InitState == EInstallBundleManagerInitState::Succeeded);
+
+	FInstallBundleSourceBundleInfoQueryResult ResultInfo;
+
+	const FConfigFile* InstallBundleConfig = GConfig->FindConfigFile(GInstallBundleIni);
+	if (InstallBundleConfig)
+	{
+		for (const TPair<FString, FConfigSection>& Pair : *InstallBundleConfig)
+		{
+			const FString& Section = Pair.Key;
+			FInstallBundleSourcePersistentBundleInfo SourceBundleInfo;
+			if (!InstallBundleManagerUtil::LoadBundleSourceBundleInfoFromConfig(GetSourceType(), *InstallBundleConfig, Section, SourceBundleInfo))
+				continue;
+
+			if (!QueryPersistentBundleInfo(SourceBundleInfo))
+				continue;
+
+			FName BundleName = SourceBundleInfo.BundleName;
+			ResultInfo.SourceBundleInfoMap.Add(BundleName, MoveTemp(SourceBundleInfo));
+		}
+	}
+
+	Callback.ExecuteIfBound(AsShared(), MoveTemp(ResultInfo));
+}
+
+EInstallBundleManagerInitState FInstallBundleSourcePlatformBase::GetInitState() const
+{
+	return InitState;
+}
+
+FString FInstallBundleSourcePlatformBase::GetContentVersion() const
+{
+	return InstallBundleUtil::GetAppVersion();
+}
+
+TSet<FName> FInstallBundleSourcePlatformBase::GetBundleDependencies(FName InBundleName, TSet<FName>* SkippedUnknownBundles /*= nullptr*/) const
+{
+	return InstallBundleManagerUtil::GetBundleDependenciesFromConfig(InBundleName, SkippedUnknownBundles);
+}
+
diff --git a/Engine/Plugins/Experimental/DefaultInstallBundleManager/Source/Public/InstallBundleSourcePlatformBase.h b/Engine/Plugins/Experimental/DefaultInstallBundleManager/Source/Public/InstallBundleSourcePlatformBase.h
new file mode 100644
--- /dev/null
+++ b/Engine/Plugins/Experimental/DefaultInstallBundleManager/Source/Public/InstallBundleSourcePlatformBase.h
@@ -0,0 +1,64 @@
+// Copyright Epic Games, Inc. All Rights Reserved.
+
+#pragma once
+
+#include "InstallBundleSourceInterface.h"
+
+#include "Internationalization/Regex.h"
+
+
+class DEFAULTINSTALLBUNDLEMANAGER_API FInstallBundleSourcePlatformBase : public IInstallBundleSource
+{
+private:
+public:
+	FInstallBundleSourcePlatformBase();
+	FInstallBundleSourcePlatformBase(const FInstallBundleSourcePlatformBase& Other) = delete;
+	FInstallBundleSourcePlatformBase& operator=(const FInstallBundleSourcePlatformBase& Other) = delete;
+	virtual ~FInstallBundleSourcePlatformBase();
+
+
+	// IInstallBundleSource Interface
+public:
+	virtual EInstallBundleSourceType GetSourceType() const override;
+
+	virtual FInstallBundleSourceInitInfo Init(
+		TSharedRef<InstallBundleUtil::FContentRequestStatsMap> InRequestStats,
+		TSharedPtr<IAnalyticsProviderET> AnalyticsProvider,
+		TSharedPtr<InstallBundleUtil::PersistentStats::FPersistentStatContainerBase> InPersistentStatsContainer) override;
+
+
+	virtual void AsyncInit_QueryBundleInfo(FInstallBundleSourceQueryBundleInfoDelegate Callback) override;
+
+	virtual EInstallBundleManagerInitState GetInitState() const override;
+
+	virtual FString GetContentVersion() const override;
+
+	virtual TSet<FName> GetBundleDependencies(FName InBundleName, TSet<FName>* SkippedUnknownBundles /*= nullptr*/) const override;
+
+
+protected:
+	virtual bool IsPlatformInstallBundleConfig(const FConfigFile* InstallBundleConfig, const FString& Section) const = 0;
+	virtual bool QueryPersistentBundleInfo(FInstallBundleSourcePersistentBundleInfo& SourceBundleInfo) const = 0;
+
+	template<typename T>
+	inline void StatsBegin(FName BundleName, T State)
+	{
+		RequestStats->StatsBegin(BundleName, LexToString(State));
+	}
+	
+	template<typename T>
+	inline void StatsEnd(FName BundleName, T State, uint64 DataSize = 0)
+	{
+		RequestStats->StatsEnd(BundleName, LexToString(State), DataSize);
+	}
+
+	TSharedPtr<IAnalyticsProviderET> AnalyticsProvider;
+
+	TSharedPtr<InstallBundleUtil::FContentRequestStatsMap> RequestStats;
+
+	TSharedPtr<InstallBundleUtil::PersistentStats::FPersistentStatContainerBase> PersistentStatsContainer;
+
+	EInstallBundleManagerInitState InitState = EInstallBundleManagerInitState::NotInitialized;
+
+	TArray<TPair<FString, TArray<FRegexPattern>>> BundleRegexList; // BundleName -> FileRegex
+};
