fixup regarding recent private pch changes in this module.

#rnx

diff --git a/Engine/Plugins/Experimental/DefaultInstallBundleManager/Source/Private/InstallBundleSourcePlatformBase.cpp b/Engine/Plugins/Experimental/DefaultInstallBundleManager/Source/Private/InstallBundleSourcePlatformBase.cpp
--- a/Engine/Plugins/Experimental/DefaultInstallBundleManager/Source/Private/InstallBundleSourcePlatformBase.cpp
+++ b/Engine/Plugins/Experimental/DefaultInstallBundleManager/Source/Private/InstallBundleSourcePlatformBase.cpp
@@ -1,8 +1,8 @@
 // Copyright Epic Games, Inc. All Rights Reserved.
 
 #include "InstallBundleSourcePlatformBase.h"
-#include "DefaultInstallBundleManagerPrivatePCH.h"
 
+#include "DefaultInstallBundleManagerPrivate.h"
 #include "InstallBundleManagerUtil.h"
 
 #include "Misc/ConfigCacheIni.h"
