fix for platform install bundle source compilation if there is no platform bundle source.

#rnx

diff --git a/Engine/Plugins/Experimental/DefaultInstallBundleManager/Source/Private/InstallBundleSourcePlatformBase.cpp b/Engine/Plugins/Experimental/DefaultInstallBundleManager/Source/Private/InstallBundleSourcePlatformBase.cpp
--- a/Engine/Plugins/Experimental/DefaultInstallBundleManager/Source/Private/InstallBundleSourcePlatformBase.cpp
+++ b/Engine/Plugins/Experimental/DefaultInstallBundleManager/Source/Private/InstallBundleSourcePlatformBase.cpp
@@ -2,6 +2,7 @@
 
 #include "InstallBundleSourcePlatformBase.h"
 
+#if WITH_PLATFORM_INSTALL_BUNDLE_SOURCE
 #include "DefaultInstallBundleManagerPrivate.h"
 #include "InstallBundleManagerUtil.h"
 
@@ -142,3 +143,4 @@ TSet<FName> FInstallBundleSourcePlatformBase::GetBundleDependencies(FName InBund
 	return InstallBundleManagerUtil::GetBundleDependenciesFromConfig(InBundleName, SkippedUnknownBundles);
 }
 
+#endif //WITH_PLATFORM_INSTALL_BUNDLE_SOURCE
diff --git a/Engine/Plugins/Experimental/DefaultInstallBundleManager/Source/Public/InstallBundleSourcePlatformBase.h b/Engine/Plugins/Experimental/DefaultInstallBundleManager/Source/Public/InstallBundleSourcePlatformBase.h
--- a/Engine/Plugins/Experimental/DefaultInstallBundleManager/Source/Public/InstallBundleSourcePlatformBase.h
+++ b/Engine/Plugins/Experimental/DefaultInstallBundleManager/Source/Public/InstallBundleSourcePlatformBase.h
@@ -3,7 +3,7 @@
 #pragma once
 
 #include "InstallBundleSourceInterface.h"
-
+#if WITH_PLATFORM_INSTALL_BUNDLE_SOURCE
 #include "Internationalization/Regex.h"
 
 
@@ -62,3 +62,4 @@ protected:
 
 	TArray<TPair<FString, TArray<FRegexPattern>>> BundleRegexList; // BundleName -> FileRegex
 };
+#endif //WITH_PLATFORM_INSTALL_BUNDLE_SOURCE
