Use a factory function instead of a typedef to create platform bundle source.

This allows additional logic, such as creating a different bundle source depending on how the game has been packaged etc.

#rnx

diff --git a/Engine/Plugins/Experimental/DefaultInstallBundleManager/Source/Private/InstallBundleManagerUtil.cpp b/Engine/Plugins/Experimental/DefaultInstallBundleManager/Source/Private/InstallBundleManagerUtil.cpp
--- a/Engine/Plugins/Experimental/DefaultInstallBundleManager/Source/Private/InstallBundleManagerUtil.cpp
+++ b/Engine/Plugins/Experimental/DefaultInstallBundleManager/Source/Private/InstallBundleManagerUtil.cpp
@@ -19,10 +19,6 @@
 	#define WITH_PLATFORM_INSTALL_BUNDLE_SOURCE 0
 #endif
 
-#if WITH_PLATFORM_INSTALL_BUNDLE_SOURCE
-	#include "PlatformInstallBundleSource.h"
-#endif
-
 #include "AnalyticsEventAttribute.h"
 #include "IAnalyticsProviderET.h"
 
@@ -51,7 +47,7 @@ namespace InstallBundleManagerUtil
 			
 #if WITH_PLATFORM_INSTALL_BUNDLE_SOURCE
 		case EInstallBundleSourceType::Platform:
-			return MakeShared<FPlatformInstallBundleSource>();
+			return MakePlatformBundleSource();
 #endif
 		}
 
diff --git a/Engine/Plugins/Experimental/DefaultInstallBundleManager/Source/Public/InstallBundleManagerUtil.h b/Engine/Plugins/Experimental/DefaultInstallBundleManager/Source/Public/InstallBundleManagerUtil.h
--- a/Engine/Plugins/Experimental/DefaultInstallBundleManager/Source/Public/InstallBundleManagerUtil.h
+++ b/Engine/Plugins/Experimental/DefaultInstallBundleManager/Source/Public/InstallBundleManagerUtil.h
@@ -23,6 +23,11 @@ namespace InstallBundleManagerUtil
 {
 	DEFAULTINSTALLBUNDLEMANAGER_API TSharedPtr<IInstallBundleSource> MakeBundleSource(EInstallBundleSourceType Type);
 
+#if WITH_PLATFORM_INSTALL_BUNDLE_SOURCE
+	DEFAULTINSTALLBUNDLEMANAGER_API TSharedPtr<IInstallBundleSource> MakePlatformBundleSource();
+#endif
+
+
 	// Returns a thread pool with one thread suitible for running in-order journal tasks
 	DEFAULTINSTALLBUNDLEMANAGER_API TSharedPtr<FQueuedThreadPool, ESPMode::ThreadSafe> GetJournalThreadPool();
 
