Install Bundle Manager
- Use an explicit list for bundle infos that can be pruned because they are no longer relevant. This prevents having to iterate the entire bundle info list which can be expensive when the list is large.

diff --git a/Engine/Plugins/Experimental/DefaultInstallBundleManager/Source/Private/DefaultInstallBundleManager.cpp b/Engine/Plugins/Experimental/DefaultInstallBundleManager/Source/Private/DefaultInstallBundleManager.cpp
--- a/Engine/Plugins/Experimental/DefaultInstallBundleManager/Source/Private/DefaultInstallBundleManager.cpp
+++ b/Engine/Plugins/Experimental/DefaultInstallBundleManager/Source/Private/DefaultInstallBundleManager.cpp
@@ -2097,14 +2097,18 @@ void FDefaultInstallBundleManager::TickPruneBundleInfo()
 	// Unfortunately, there isn't any good way to do this other than on tick because
 	// we need to have no requests for a bundles in flight when we prune it.
 
-	for (auto It = BundleInfoMap.CreateIterator(); It; ++It)
+	for (auto It = BundlesInfosToPrune.CreateIterator(); It; ++It)
 	{
-		const FName BundleName = It.Key();
-		FBundleInfo& BundleInfo = It.Value();
+		const FName BundleName = *It;
+		const FBundleInfo& BundleInfo = BundleInfoMap.FindChecked(BundleName);
 
 		const bool bIsRelevant = Algo::AnyOf(BundleInfo.ContributingSources, &FBundleSourceRelevance::bIsRelevant);
 		if (bIsRelevant)
+		{
+			ensureAlwaysMsgf(false, TEXT("TickPruneBundleInfo - Found relevant bundle in prune list %s"), *BundleInfo.BundleNameString);
+			It.RemoveCurrent();
 			continue;
+		}
 
 		bool bIsRequested = false;
 		for (EContentRequestBatch b : TEnumRange<EContentRequestBatch>())
@@ -2159,6 +2163,7 @@ void FDefaultInstallBundleManager::TickPruneBundleInfo()
 			}
 		}
 
+		BundleInfoMap.Remove(BundleName);
 		It.RemoveCurrent();
 	}
 }
@@ -2725,6 +2730,9 @@ EInstallBundleSourceUpdateBundleInfoResult FDefaultInstallBundleManager::OnUpdat
 		{
 			BundleInfo.ContributingSources.Add(FBundleSourceRelevance{ SourceType, true });
 		}
+
+		// Since the bundle is now relevant make sure its not in the prune list
+		BundlesInfosToPrune.Remove(UpdateInfoPair.Key);
 	}
 
 	return EInstallBundleSourceUpdateBundleInfoResult::OK;
@@ -2751,6 +2759,13 @@ void FDefaultInstallBundleManager::OnBundleLostRelevanceForSource(TSharedRef<IIn
 		}
 
 		SourceRelevance->bIsRelevant = false;
+
+		// See if all relevance has been lost. If so, add to the set of bundle infos to be pruned
+		const bool bIsRelevant = Algo::AnyOf(BundleInfo->ContributingSources, &FBundleSourceRelevance::bIsRelevant);
+		if (!bIsRelevant)
+		{
+			BundlesInfosToPrune.Add(BundleName);
+		}
 	}
 }
 
diff --git a/Engine/Plugins/Experimental/DefaultInstallBundleManager/Source/Public/DefaultInstallBundleManager.h b/Engine/Plugins/Experimental/DefaultInstallBundleManager/Source/Public/DefaultInstallBundleManager.h
--- a/Engine/Plugins/Experimental/DefaultInstallBundleManager/Source/Public/DefaultInstallBundleManager.h
+++ b/Engine/Plugins/Experimental/DefaultInstallBundleManager/Source/Public/DefaultInstallBundleManager.h
@@ -706,6 +706,7 @@ protected:
 	FInstallBundleSourceFactoryFunction InstallBundleSourceFactory;
 
 	TMap<FName, FBundleInfo> BundleInfoMap;
+	TSet<FName> BundlesInfosToPrune;
 
 	TMap<EInstallBundleSourceType, TSharedPtr<IInstallBundleSource>> BundleSources;
 	TMap<EInstallBundleSourceType, EInstallBundleSourceType> BundleSourceFallbacks;
