Fix analytics provider not being available during install bundle source initialization.
This prevented some analytics events from being sent.

diff --git a/Engine/Plugins/Experimental/DefaultInstallBundleManager/Source/Private/DefaultInstallBundleManager.cpp b/Engine/Plugins/Experimental/DefaultInstallBundleManager/Source/Private/DefaultInstallBundleManager.cpp
--- a/Engine/Plugins/Experimental/DefaultInstallBundleManager/Source/Private/DefaultInstallBundleManager.cpp
+++ b/Engine/Plugins/Experimental/DefaultInstallBundleManager/Source/Private/DefaultInstallBundleManager.cpp
@@ -201,16 +201,6 @@ FDefaultInstallBundleManager::FDefaultInstallBundleManager(const TCHAR* InConfig
 	const TCHAR* DefaultConfigBaseName = TEXT("InstallBundleManager");
 	FConfigContext::ReadIntoGConfig().Load(InConfigBaseName ? InConfigBaseName : DefaultConfigBaseName, GInstallBundleManagerIni);
 
-	TickHandle = FTSTicker::GetCoreTicker().AddTicker(FTickerDelegate::CreateRaw(this, &FDefaultInstallBundleManager::Tick));
-
-	InitErrorHandlerStack.Push(FInstallBundleManagerInitErrorHandler::CreateLambda(
-		[](EInstallBundleManagerInitResult Error)
-		{
-			// Default Handler, Just keep retrying
-			return EInstallBundleManagerInitErrorHandlerResult::Retry;
-		}
-	));
-
 #if WITH_EDITOR
 	// -UsePaks needs to be specified on the command line for valid pak to be created.
 	// To support mounting pak files in the editor binary add the encryption key.
@@ -221,6 +211,27 @@ FDefaultInstallBundleManager::FDefaultInstallBundleManager(const TCHAR* InConfig
 
 	SetErrorSimulationCommands(FCommandLine::Get());
 	SetCommandLineOverrides(FCommandLine::Get());
+}
+
+FDefaultInstallBundleManager::~FDefaultInstallBundleManager()
+{
+	FTSTicker::GetCoreTicker().RemoveTicker(TickHandle);
+	TickHandle.Reset();
+
+	InstallBundleUtil::CleanupInstallBundleAsyncIOTasks(AsyncMountTasks);
+}
+
+void FDefaultInstallBundleManager::Initialize()
+{
+	TickHandle = FTSTicker::GetCoreTicker().AddTicker(FTickerDelegate::CreateRaw(this, &FDefaultInstallBundleManager::Tick));
+
+	InitErrorHandlerStack.Push(FInstallBundleManagerInitErrorHandler::CreateLambda(
+		[](EInstallBundleManagerInitResult Error)
+		{
+			// Default Handler, Just keep retrying
+			return EInstallBundleManagerInitErrorHandlerResult::Retry;
+		}
+	));
 
 	if (InitResult == EInstallBundleManagerInitResult::OK)
 		InitResult = Init_DefaultBundleSources();
@@ -232,14 +243,6 @@ FDefaultInstallBundleManager::FDefaultInstallBundleManager(const TCHAR* InConfig
 	}
 }
 
-FDefaultInstallBundleManager::~FDefaultInstallBundleManager()
-{
-	FTSTicker::GetCoreTicker().RemoveTicker(TickHandle);
-	TickHandle.Reset();
-
-	InstallBundleUtil::CleanupInstallBundleAsyncIOTasks(AsyncMountTasks);
-}
-
 bool FDefaultInstallBundleManager::Tick(float dt)
 {
 	QUICK_SCOPE_CYCLE_COUNTER(STAT_FDefaultInstallBundleManager_Tick);
diff --git a/Engine/Plugins/Experimental/DefaultInstallBundleManager/Source/Public/DefaultInstallBundleManager.h b/Engine/Plugins/Experimental/DefaultInstallBundleManager/Source/Public/DefaultInstallBundleManager.h
--- a/Engine/Plugins/Experimental/DefaultInstallBundleManager/Source/Public/DefaultInstallBundleManager.h
+++ b/Engine/Plugins/Experimental/DefaultInstallBundleManager/Source/Public/DefaultInstallBundleManager.h
@@ -485,6 +485,8 @@ public:
 
 	virtual ~FDefaultInstallBundleManager();
 
+	virtual void Initialize() override;
+
 	//Tick
 protected:
 	bool Tick(float dt);
diff --git a/Engine/Source/Runtime/InstallBundleManager/Public/InstallBundleManagerInterface.h b/Engine/Source/Runtime/InstallBundleManager/Public/InstallBundleManagerInterface.h
--- a/Engine/Source/Runtime/InstallBundleManager/Public/InstallBundleManagerInterface.h
+++ b/Engine/Source/Runtime/InstallBundleManager/Public/InstallBundleManagerInterface.h
@@ -96,6 +96,8 @@ public:
 
 	virtual ~IInstallBundleManager() {}
 
+	virtual void Initialize() {}
+
 	virtual bool HasBundleSource(EInstallBundleSourceType SourceType) const = 0;
 
 	virtual FDelegateHandle PushInitErrorCallback(FInstallBundleManagerInitErrorHandler Callback) = 0;
diff --git a/Engine/Source/Runtime/InstallBundleManager/Public/InstallBundleManagerModule.h b/Engine/Source/Runtime/InstallBundleManager/Public/InstallBundleManagerModule.h
--- a/Engine/Source/Runtime/InstallBundleManager/Public/InstallBundleManagerModule.h
+++ b/Engine/Source/Runtime/InstallBundleManager/Public/InstallBundleManagerModule.h
@@ -72,6 +72,8 @@ public:
 			check(LastInstallBundleManager == nullptr);
 			InstallBundleManager = MakeShared<InstallBundleManagerModuleImpl>();
 			LastInstallBundleManager = InstallBundleManager;
+
+			InstallBundleManager->Initialize();
 		}
 	}
 };
