Speculative crash fix for bundle info being pruned after a content state request is enqueued.

-review-22485073 [REDACTED_HANDLE]

diff --git a/Engine/Plugins/Experimental/DefaultInstallBundleManager/Source/Private/DefaultInstallBundleManager.cpp b/Engine/Plugins/Experimental/DefaultInstallBundleManager/Source/Private/DefaultInstallBundleManager.cpp
--- a/Engine/Plugins/Experimental/DefaultInstallBundleManager/Source/Private/DefaultInstallBundleManager.cpp
+++ b/Engine/Plugins/Experimental/DefaultInstallBundleManager/Source/Private/DefaultInstallBundleManager.cpp
@@ -2116,7 +2116,7 @@ void FDefaultInstallBundleManager::TickPruneBundleInfo()
 		bool bIsRequested = false;
 		for (EContentRequestBatch b : TEnumRange<EContentRequestBatch>())
 		{
-			for (FContentRequestRef& QueuedRequest : ContentRequests[b])
+			for (const FContentRequestRef& QueuedRequest : ContentRequests[b])
 			{
 				if (QueuedRequest->BundleName == BundleName)
 				{
@@ -2126,15 +2126,15 @@ void FDefaultInstallBundleManager::TickPruneBundleInfo()
 			}
 
 			if (bIsRequested)
-				break;
+			{	break; }
 		}
 
 		if (bIsRequested)
-			continue;
+		{	continue; }
 
 		for (EContentReleaseRequestBatch b : TEnumRange<EContentReleaseRequestBatch>())
 		{
-			for (FContentReleaseRequestRef& QueuedRequest : ContentReleaseRequests[b])
+			for (const FContentReleaseRequestRef& QueuedRequest : ContentReleaseRequests[b])
 			{
 				if (QueuedRequest->BundleName == BundleName)
 				{
@@ -2144,11 +2144,35 @@ void FDefaultInstallBundleManager::TickPruneBundleInfo()
 			}
 
 			if (bIsRequested)
+			{	break; }
+		}
+
+		if (bIsRequested)
+		{	continue; }
+
+		for (const FGetContentStateRequestRef& Request : GetContentStateRequests)
+		{
+			if (Request->BundleNames.Contains(BundleName))
+			{
+				bIsRequested = true;
+				break;
+			}
+		}
+
+		if (bIsRequested)
+		{	continue; }
+
+		for (const FGetInstallStateRequestRef& Request : GetInstallStateRequests)
+		{
+			if (Request->BundleNames.Contains(BundleName))
+			{
+				bIsRequested = true;
 				break;
+			}
 		}
 
 		if (bIsRequested)
-			continue;
+		{	continue; }
 
 		for (const FBundleSourceRelevance& SourceRelevance : BundleInfo.ContributingSources)
 		{
