Changes to drastically speed up Android bulk build iteration when using InstallBundleManager with a large number of paks.
#test Android Dev Build

diff --git a/Engine/Config/Android/BaseAndroidInstallBundle.ini b/Engine/Config/Android/BaseAndroidInstallBundle.ini
new file mode 100644
--- /dev/null
+++ b/Engine/Config/Android/BaseAndroidInstallBundle.ini
@@ -0,0 +1,7 @@
+[InstallBundleSource.Bulk.MiscSettings]
+; Allows our Bulk data source to serialize a BulkBuildData.ini if it was missing during loading.
+; Can drastically speed up Bulk data source initialization for runs after the first as it won't have to compute
+; Bulk bundle information during init.
+bShouldSerializeMissingBulkBuildDataIni=true
+; If we limit our file parsing in Bulk bundle sources to just .pak files
+bOnlyGatherPaksInBulkData=true
\ No newline at end of file
diff --git a/Engine/Config/BaseInstallBundle.ini b/Engine/Config/BaseInstallBundle.ini
--- a/Engine/Config/BaseInstallBundle.ini
+++ b/Engine/Config/BaseInstallBundle.ini
@@ -1 +1,10 @@
-; Defines install bundles and options
\ No newline at end of file
+; Defines install bundles and options
+
+
+[InstallBundleSource.Bulk.MiscSettings]
+; Allows our Bulk data source to serialize a BulkBuildData.ini if it was missing during loading.
+; Can drastically speed up Bulk data source initialization for runs after the first as it won't have to compute
+; Bulk bundle information during init.
+bShouldSerializeMissingBulkBuildDataIni=false
+; If we limit our file parsing in Bulk bundle sources to just .pak files
+bOnlyGatherPaksInBulkData=false
\ No newline at end of file
diff --git a/Engine/Plugins/Experimental/DefaultInstallBundleManager/Source/Private/InstallBundleSourceBulk.cpp b/Engine/Plugins/Experimental/DefaultInstallBundleManager/Source/Private/InstallBundleSourceBulk.cpp
--- a/Engine/Plugins/Experimental/DefaultInstallBundleManager/Source/Private/InstallBundleSourceBulk.cpp
+++ b/Engine/Plugins/Experimental/DefaultInstallBundleManager/Source/Private/InstallBundleSourceBulk.cpp
@@ -10,11 +10,143 @@
 #include "HAL/FileManager.h"
 #include "Internationalization/Regex.h"
 #include "Stats/Stats.h"
+#include "Serialization/JsonSerializerMacros.h"
 
 #define LOG_SOURCE_BULK(Verbosity, Format, ...) LOG_INSTALL_BUNDLE_MAN(Verbosity, TEXT("InstallBundleSourceBulk: ") Format, ##__VA_ARGS__)
 
 #define LOG_SOURCE_BULK_OVERRIDE(VerbosityOverride, Verbosity, Format, ...) LOG_INSTALL_BUNDLE_MAN_OVERRIDE(VerbosityOverride, Verbosity, TEXT("InstallBundleSourceBulk: ") Format, ##__VA_ARGS__)
 
+//Helper class used to load/save BulkBuildBundle information from/to disk
+class FBulkBuildBundleMapJsonInfo
+	: public FJsonSerializable
+{
+public:
+
+	//Simple Wrapper to hold bundle name in a way we can serialize it to/from a map with JSON_SERIALIZE_MAP_SERIALIZABLE
+	class FJsonBundleNameWrapper
+		: public FJsonSerializable
+	{
+		public:
+			BEGIN_JSON_SERIALIZER
+				JSON_SERIALIZE("BundleName", BundleName);
+			END_JSON_SERIALIZER
+
+			FString BundleName;
+			
+			FJsonBundleNameWrapper()
+				: BundleName()
+			{}
+	};
+	
+	BEGIN_JSON_SERIALIZER
+		JSON_SERIALIZE_MAP_SERIALIZABLE("BulkBuildBundleByFileMap", BulkBuildBundleByFileMap, FJsonBundleNameWrapper);
+	END_JSON_SERIALIZER
+
+	bool LoadFromFile(FStringView FilePath)
+	{
+		FString JSONStringOnDisk;
+		if (FPaths::FileExists(FilePath.GetData()))
+		{
+			FFileHelper::LoadFileToString(JSONStringOnDisk, FilePath.GetData());
+		}
+
+		if (!JSONStringOnDisk.IsEmpty())
+		{
+			return ensureAlwaysMsgf(
+				FromJson(JSONStringOnDisk),
+				TEXT("Invalid JSON found while parsing BulkBuildBundleMapInfo from JSON: %s loaded from file:%.*s"), 
+				*JSONStringOnDisk,
+				FilePath.Len(), FilePath.GetData());
+		}
+
+		return false;
+	}
+
+	bool SaveToFile(FStringView FilePath)
+	{
+		return ensureAlwaysMsgf(
+				FFileHelper::SaveStringToFile(ToJson(), FilePath.GetData()),
+				TEXT("Error saving Json output of FBulkBuildBundleMapInfo to %.*s"),
+				FilePath.Len(),FilePath.GetData());
+	}
+
+	//Takes in a list of files found on disk and uses our parsed BulkBuildBundleByFileMap to fill out the OutBulkBuildBundles Map.
+	//Removes all found entries from the OutBundleFileList that were successfully processed.
+	bool AppendEntriesToBulkBuildBundleMap(TArray<FString>& InOutBundleFileList, TMap<FName, TArray<FString>>& InOutBulkBuildBundles)
+	{
+		int32 OriginalBulkBuildBundleNum = InOutBulkBuildBundles.Num();
+
+		if (InOutBundleFileList.Num() == 0)
+		{
+			return false;
+		}
+
+		for (int FileIndex = 0; FileIndex < InOutBundleFileList.Num();)
+		{
+			FString& FileInList = InOutBundleFileList[FileIndex];
+			FJsonBundleNameWrapper* FoundBundleNameString = BulkBuildBundleByFileMap.Find(FileInList);
+			if (FoundBundleNameString)
+			{
+				FName BundleName(*(FoundBundleNameString->BundleName));
+				TArray<FString>& FoundFileList = InOutBulkBuildBundles.FindOrAdd(BundleName);
+				FoundFileList.AddUnique(FileInList);
+
+				//Remove and don't increment FileIndex so that we re-check this swapped index next pass if valid
+				InOutBundleFileList.RemoveAtSwap(FileIndex);
+			}
+			else
+			{
+				++FileIndex;
+			}
+		}
+
+		//return true if we've added values to OutBulkBuildBundles
+		return (InOutBulkBuildBundles.Num() > OriginalBulkBuildBundleNum);
+	}
+
+	bool IsEmpty()
+	{
+		return (BulkBuildBundleByFileMap.Num() == 0);
+	}
+
+	//Gets the path to use for the BulkBuildInfo file if its present in the cooked data on device
+	static FStringView GetBulkBuildBundleInfoCookedPath()
+	{
+		static FString CookedPath = FPaths::Combine(FPaths::ProjectContentDir(), TEXT("BulkBuildMeta"), TEXT("CachedBuilkBuildBuildInfo.json"));
+		return CookedPath;
+	}
+
+	//Gets the path to use for the BulkBuildInfo file if its cached locally on the device
+	static FStringView GetBulkBuildBundleInfoLocalCachedPath()
+	{
+		static FString LocalCachedPath = FPaths::Combine(FPaths::ProjectUserDir(), TEXT("Saved"), TEXT("BulkBuildMeta"), TEXT("CachedBuilkBuildBuildInfo.json"));
+		return LocalCachedPath;
+	}
+
+	//Constructor to generate BulkBuildBundleByFileMap from BulkBuildBundlesIn, useful for loading the BulkBuildBundle info from
+	//memory and then saving it out to file
+	FBulkBuildBundleMapJsonInfo(const TMap<FName, TArray<FString>>& BulkBuildBundlesIn)
+		: FBulkBuildBundleMapJsonInfo()
+	{
+		for(const TPair<FName, TArray<FString>>& BundlePair : BulkBuildBundlesIn)
+		{
+			for (const FString& FileName : BundlePair.Value)
+			{
+				FJsonBundleNameWrapper& FoundBundleName = BulkBuildBundleByFileMap.FindOrAdd(FileName);
+				FoundBundleName.BundleName = BundlePair.Key.ToString();
+			}
+		}
+	}
+
+	FBulkBuildBundleMapJsonInfo()
+		: BulkBuildBundleByFileMap()
+	{};
+
+private:
+	//Serialized BulkBuildBundle (FileName)->(Bundle FString version of FName) information
+	TMap<FString, FJsonBundleNameWrapper> BulkBuildBundleByFileMap;
+};
+
 FInstallBundleSourceBulk::FInstallBundleSourceBulk()
 {
 	TickHandle = FTSTicker::GetCoreTicker().AddTicker(FTickerDelegate::CreateRaw(this, &FInstallBundleSourceBulk::Tick));
@@ -122,13 +254,23 @@ void FInstallBundleSourceBulk::AsyncInit_MakeBundlesForBulkBuild()
 	TArray<FString> PakSearchDirs;
 	FPakPlatformFile::GetPakFolders(FCommandLine::Get(), PakSearchDirs);
 
+	//Get setting for if we limit our file list to only .pak files
+	bool bOnlyGatherPaksInBulkData = false;
+	if (!GConfig->GetBool(TEXT("InstallBundleSource.Bulk.MiscSettings"), TEXT("bOnlyGatherPaksInBulkData"), bOnlyGatherPaksInBulkData, GInstallBundleIni))
+	{
+		bOnlyGatherPaksInBulkData = false;
+	}
+
 	TSharedPtr<TArray<FString>, ESPMode::ThreadSafe> FoundFiles = MakeShared<TArray<FString>, ESPMode::ThreadSafe>();
 	InstallBundleUtil::StartInstallBundleAsyncIOTask(InitAsyncTasks,
-	[FoundFiles, PakSearchDirs=MoveTemp(PakSearchDirs), ContentDir=FPaths::ProjectContentDir()]()
+	[FoundFiles, PakSearchDirs=MoveTemp(PakSearchDirs), ContentDir=FPaths::ProjectContentDir(), bOnlyGatherPaksInBulkData]()
 	{
+		const FString PakFileExtension(TEXT(".pak"));
+		const TCHAR* FileExtension = bOnlyGatherPaksInBulkData ? *PakFileExtension : nullptr;
+		
 		for (const FString& SearchDir : PakSearchDirs)
 		{
-			IPlatformFile::GetPlatformPhysical().FindFilesRecursively(*FoundFiles, *SearchDir, nullptr);
+			IPlatformFile::GetPlatformPhysical().FindFilesRecursively(*FoundFiles, *SearchDir, FileExtension);
 		}
 
 #if PLATFORM_IOS
@@ -139,7 +281,15 @@ void FInstallBundleSourceBulk::AsyncInit_MakeBundlesForBulkBuild()
 	},
 	[this, FoundFiles]()
 	{
-		// Prune out any paks that were mounted at startup
+		if (!FoundFiles.IsValid())
+		{
+			return;
+		}
+
+		//First load any existing BulkBuildBundleIni entries
+		bool bDidLoadAllIniEntries = TryLoadBulkBuildBundleMetadata(*FoundFiles, BulkBuildBundles);
+
+		// Prune out any remaining paks that were mounted at startup
 		for (int i = 0; i < FoundFiles->Num();)
 		{
 			const FString& File = (*FoundFiles)[i];
@@ -153,86 +303,140 @@ void FInstallBundleSourceBulk::AsyncInit_MakeBundlesForBulkBuild()
 			}
 		}
 
-		TArray<FString> SectionNames;
-		const FConfigFile* InstallBundleConfig = GConfig->FindConfigFile(GInstallBundleIni);
-		if (InstallBundleConfig)
+		LOG_SOURCE_BULK(Display, TEXT("Loaded %d Bundle Information from BulkBuildBundles Cache. %d Files Remaining."), BulkBuildBundles.Num(), FoundFiles->Num());
+
+		//Only manually parse if there are remaining FoundFiles
+		if (FoundFiles->Num() > 0)
 		{
-			for (const TPair<FString, FConfigSection>& Pair : *InstallBundleConfig)
-			{
-				if (Pair.Key.StartsWith(InstallBundleUtil::GetInstallBundleSectionPrefix()))
-				{
-					SectionNames.Add(Pair.Key);
-				}
-			}
+			bDidLoadAllIniEntries = false;
 
-			// Bundle regex need to be applied in order
-			SectionNames.StableSort([this](const FString& SectionA, const FString& SectionB) -> bool
+			TArray<FString> SectionNames;
+			const FConfigFile* InstallBundleConfig = GConfig->FindConfigFile(GInstallBundleIni);
+			if (InstallBundleConfig)
 			{
-				int BundleAOrder = INT_MAX;
-				int BundleBOrder = INT_MAX;
-
-				if (!GConfig->GetInt(*SectionA, TEXT("Order"), BundleAOrder, GInstallBundleIni))
+				for (const TPair<FString, FConfigSection>& Pair : *InstallBundleConfig)
 				{
-					LOG_SOURCE_BULK(Warning, TEXT("Bundle Section %s doesn't have an order"), *SectionA);
+					if (Pair.Key.StartsWith(InstallBundleUtil::GetInstallBundleSectionPrefix()))
+					{
+						SectionNames.Add(Pair.Key);
+					}
 				}
 
-				if (!GConfig->GetInt(*SectionB, TEXT("Order"), BundleBOrder, GInstallBundleIni))
+				// Bundle regex need to be applied in order
+				SectionNames.StableSort([this](const FString& SectionA, const FString& SectionB) -> bool
 				{
-					LOG_SOURCE_BULK(Warning, TEXT("Bundle Section %s doesn't have an order"), *SectionB);
-				}
+					int BundleAOrder = INT_MAX;
+					int BundleBOrder = INT_MAX;
 
-				return BundleAOrder < BundleBOrder;
-			});
+					if (!GConfig->GetInt(*SectionA, TEXT("Order"), BundleAOrder, GInstallBundleIni))
+					{
+						LOG_SOURCE_BULK(Warning, TEXT("Bundle Section %s doesn't have an order"), *SectionA);
+					}
 
-			// Sort remaining files into bundles
-			for (const FString& Section : SectionNames)
-			{
-				const FString BundleName = Section.RightChop(InstallBundleUtil::GetInstallBundleSectionPrefix().Len());
-				TArray<FString>& BundleFileList = BulkBuildBundles.FindOrAdd(FName(*BundleName));
+					if (!GConfig->GetInt(*SectionB, TEXT("Order"), BundleBOrder, GInstallBundleIni))
+					{
+						LOG_SOURCE_BULK(Warning, TEXT("Bundle Section %s doesn't have an order"), *SectionB);
+					}
 
-				TArray<FString> StrSearchRegexPatterns;
-				if (!InstallBundleConfig->GetArray(*Section, TEXT("FileRegex"), StrSearchRegexPatterns))
-					continue;
+					return BundleAOrder < BundleBOrder;
+				});
 
-				TArray<FRegexPattern> SearchRegexPatterns;
-				SearchRegexPatterns.Reserve(StrSearchRegexPatterns.Num());
-				for (const FString& Str : StrSearchRegexPatterns)
+				// Sort remaining files into bundles
+				for (const FString& Section : SectionNames)
 				{
-					SearchRegexPatterns.Emplace(Str, ERegexPatternFlags::CaseInsensitive);
-				}
+					const FString BundleName = Section.RightChop(InstallBundleUtil::GetInstallBundleSectionPrefix().Len());
+					TArray<FString>& BundleFileList = BulkBuildBundles.FindOrAdd(FName(*BundleName));
+
+					TArray<FString> StrSearchRegexPatterns;
+					if (!InstallBundleConfig->GetArray(*Section, TEXT("FileRegex"), StrSearchRegexPatterns))
+						continue;
+
+					TArray<FRegexPattern> SearchRegexPatterns;
+					SearchRegexPatterns.Reserve(StrSearchRegexPatterns.Num());
+					for (const FString& Str : StrSearchRegexPatterns)
+					{
+						SearchRegexPatterns.Emplace(Str, ERegexPatternFlags::CaseInsensitive);
+					}
 				
-				for (int i = 0; i < FoundFiles->Num();)
-				{
-					const FString& File = (*FoundFiles)[i];
-					bool bMatches = false;
-					for (const FRegexPattern& Pattern : SearchRegexPatterns)
+					for (int i = 0; i < FoundFiles->Num();)
 					{
-						if (FRegexMatcher(Pattern, File).FindNext())
+						const FString& File = (*FoundFiles)[i];
+						bool bMatches = false;
+						for (const FRegexPattern& Pattern : SearchRegexPatterns)
 						{
-							bMatches = true;
-							break;
+							if (FRegexMatcher(Pattern, File).FindNext())
+							{
+								bMatches = true;
+								break;
+							}
 						}
-					}
 
-					if (bMatches)
-					{
-						LOG_SOURCE_BULK(Verbose, TEXT("Adding %s to Bundle %s"), *File, *BundleName);
+						if (bMatches)
+						{
+							LOG_SOURCE_BULK(Verbose, TEXT("Adding %s to Bundle %s"), *File, *BundleName);
 						
-						BundleFileList.AddUnique(File);
-						FoundFiles->RemoveAtSwap(i);
-					}
-					else
-					{
-						++i;
+							BundleFileList.AddUnique(File);
+							FoundFiles->RemoveAtSwap(i);
+						}
+						else
+						{
+							++i;
+						}
 					}
 				}
 			}
 		}
 
+		//See if we should serialize out the manually parsed results
+		if (!bDidLoadAllIniEntries)
+		{
+			bool bShouldSerializeMissingBulkBuildDataIni = false;
+			if (!GConfig->GetBool(TEXT("InstallBundleSource.Bulk.MiscSettings"), TEXT("bShouldSerializeMissingBulkBuildDataIni"), bShouldSerializeMissingBulkBuildDataIni, GInstallBundleIni))
+			{
+				bShouldSerializeMissingBulkBuildDataIni = false;
+			}
+
+			if (bShouldSerializeMissingBulkBuildDataIni)
+			{
+				SerializeBulkBuildBundleMetadata(BulkBuildBundles);
+			}
+		}
+
+		LOG_SOURCE_BULK(Display, TEXT("Finished Making Bundles for Bulk Build"));
 		InitStepResult = EAsyncInitStepResult::Done;
 	});
 }
 
+bool FInstallBundleSourceBulk::TryLoadBulkBuildBundleMetadata(TArray<FString>& InOutFileList, TMap<FName, TArray<FString>>& InOutBulkBuildBundles)
+{
+	FBulkBuildBundleMapJsonInfo LoadedBuildInfo;
+
+	//Always prioritize loading from the LocalCache
+	if (!LoadedBuildInfo.LoadFromFile(FBulkBuildBundleMapJsonInfo::GetBulkBuildBundleInfoLocalCachedPath()))
+	{
+		//If there is no local cache look for a cooked one
+		LoadedBuildInfo.LoadFromFile(FBulkBuildBundleMapJsonInfo::GetBulkBuildBundleInfoCookedPath());
+	}
+
+	if (!LoadedBuildInfo.IsEmpty())
+	{
+		return LoadedBuildInfo.AppendEntriesToBulkBuildBundleMap(InOutFileList, InOutBulkBuildBundles);
+	}
+
+	return false;
+}
+
+void FInstallBundleSourceBulk::SerializeBulkBuildBundleMetadata(const TMap<FName, TArray<FString>>& BulkBuildBundles)
+{
+	FBulkBuildBundleMapJsonInfo JsonBulkBuildInfo(BulkBuildBundles);
+	if (!JsonBulkBuildInfo.IsEmpty())
+	{
+		FStringView FilePath = FBulkBuildBundleMapJsonInfo::GetBulkBuildBundleInfoLocalCachedPath();
+		const bool bSuccess = JsonBulkBuildInfo.SaveToFile(FilePath);
+		LOG_SOURCE_BULK(Display, TEXT("Saving BulkBuildBundle Cache to %.*s . bSuccess:%s"), FilePath.Len(), FilePath.GetData(), *LexToString(bSuccess));
+	}
+}
+
 EInstallBundleInstallState FInstallBundleSourceBulk::GetBundleInstallState(FName BundleName)
 {
 	return EInstallBundleInstallState::UpToDate;
diff --git a/Engine/Plugins/Experimental/DefaultInstallBundleManager/Source/Public/InstallBundleSourceBulk.h b/Engine/Plugins/Experimental/DefaultInstallBundleManager/Source/Public/InstallBundleSourceBulk.h
--- a/Engine/Plugins/Experimental/DefaultInstallBundleManager/Source/Public/InstallBundleSourceBulk.h
+++ b/Engine/Plugins/Experimental/DefaultInstallBundleManager/Source/Public/InstallBundleSourceBulk.h
@@ -51,6 +51,14 @@ public:
 
 	virtual void SetErrorSimulationCommands(const FString& CommandLine) override;
 
+	//Function that loads BulkBundleBuild information from a generated BulkBuildBundleIni instead of applying reg-ex at runtime
+	//Returns true if BulkBundleBuild.ini existed and was parsed successfully, false otherwise
+	//Removes any loaded entries from InOutFileList
+	static bool TryLoadBulkBuildBundleMetadata(TArray<FString>& InOutFileList, TMap<FName, TArray<FString>>& InOutBulkBuildBundles);
+
+	//Serialize out our BulkBundleBuild information to a BulkBundleBuild.ini file for future runs to not have to parse this information
+	static void SerializeBulkBuildBundleMetadata(const TMap<FName, TArray<FString>>& BulkBuildBundles);
+
 protected:
 	FTSTicker::FDelegateHandle TickHandle;
 
