[CIS] Fixing Non-Unity build error introduced in CL 23138907
#tests local Windows dev compile

diff --git a/Engine/Plugins/Experimental/DefaultInstallBundleManager/Source/Private/InstallBundleSourceBulk.cpp b/Engine/Plugins/Experimental/DefaultInstallBundleManager/Source/Private/InstallBundleSourceBulk.cpp
--- a/Engine/Plugins/Experimental/DefaultInstallBundleManager/Source/Private/InstallBundleSourceBulk.cpp
+++ b/Engine/Plugins/Experimental/DefaultInstallBundleManager/Source/Private/InstallBundleSourceBulk.cpp
@@ -6,6 +6,8 @@
 #include "IPlatformFilePak.h"
 #include "Containers/Ticker.h"
 #include "Misc/CommandLine.h"
+#include "Misc/Paths.h"
+#include "Misc/FileHelper.h"
 #include "Misc/ConfigCacheIni.h"
 #include "HAL/FileManager.h"
 #include "Internationalization/Regex.h"
