Fixing Editor crash caused from Bulk Build iteration improvements.
- Now bundle config sections are always appended to BulkBuildBundles even if they would have an empty bundle list
- Small code clean-ups
#tests Windows Dev Editor

diff --git a/Engine/Plugins/Experimental/DefaultInstallBundleManager/Source/Private/InstallBundleSourceBulk.cpp b/Engine/Plugins/Experimental/DefaultInstallBundleManager/Source/Private/InstallBundleSourceBulk.cpp
--- a/Engine/Plugins/Experimental/DefaultInstallBundleManager/Source/Private/InstallBundleSourceBulk.cpp
+++ b/Engine/Plugins/Experimental/DefaultInstallBundleManager/Source/Private/InstallBundleSourceBulk.cpp
@@ -283,13 +283,9 @@ void FInstallBundleSourceBulk::AsyncInit_MakeBundlesForBulkBuild()
 	},
 	[this, FoundFiles]()
 	{
-		if (!FoundFiles.IsValid())
-		{
-			return;
-		}
-
-		//First load any existing BulkBuildBundleIni entries
-		bool bDidLoadAllIniEntries = TryLoadBulkBuildBundleMetadata(*FoundFiles, BulkBuildBundles);
+		//First load any existing BulkBuildBundleIni entries and use those to sort FoundFiles into bundles
+		const bool bHasAnyFilesToParse = FoundFiles.IsValid() && (FoundFiles->Num() > 0);
+		bool bDidLoadAllFilesFromMetadata = bHasAnyFilesToParse ? TryLoadBulkBuildBundleMetadata(*FoundFiles, BulkBuildBundles) : false;
 
 		// Prune out any remaining paks that were mounted at startup
 		for (int i = 0; i < FoundFiles->Num();)
@@ -305,25 +301,29 @@ void FInstallBundleSourceBulk::AsyncInit_MakeBundlesForBulkBuild()
 			}
 		}
 
-		LOG_SOURCE_BULK(Display, TEXT("Loaded %d Bundle Information from BulkBuildBundles Cache. %d Files Remaining."), BulkBuildBundles.Num(), FoundFiles->Num());
-
-		//Only manually parse if there are remaining FoundFiles
+		//Still files left that weren't in the Metadata or StartupPaksWildcard, so we have files that will be manually parsed
 		if (FoundFiles->Num() > 0)
 		{
-			bDidLoadAllIniEntries = false;
+			bDidLoadAllFilesFromMetadata = false;
+		}
+
+		LOG_SOURCE_BULK(Display, TEXT("Loaded %d Bundle Information from BulkBuildBundles Cache. %d Files Remaining."), BulkBuildBundles.Num(), FoundFiles->Num());
 
-			TArray<FString> SectionNames;
-			const FConfigFile* InstallBundleConfig = GConfig->FindConfigFile(GInstallBundleIni);
-			if (InstallBundleConfig)
+		TArray<FString> SectionNames;
+		const FConfigFile* InstallBundleConfig = GConfig->FindConfigFile(GInstallBundleIni);
+		if (InstallBundleConfig)
+		{
+			for (const TPair<FString, FConfigSection>& Pair : *InstallBundleConfig)
 			{
-				for (const TPair<FString, FConfigSection>& Pair : *InstallBundleConfig)
+				if (Pair.Key.StartsWith(InstallBundleUtil::GetInstallBundleSectionPrefix()))
 				{
-					if (Pair.Key.StartsWith(InstallBundleUtil::GetInstallBundleSectionPrefix()))
-					{
-						SectionNames.Add(Pair.Key);
-					}
+					SectionNames.Add(Pair.Key);
 				}
+			}
 
+			//Skip sorting SectionNames if we have already loaded everything as we will not be applying any regex anyway
+			if (!bDidLoadAllFilesFromMetadata)
+			{
 				// Bundle regex need to be applied in order
 				SectionNames.StableSort([this](const FString& SectionA, const FString& SectionB) -> bool
 				{
@@ -342,8 +342,18 @@ void FInstallBundleSourceBulk::AsyncInit_MakeBundlesForBulkBuild()
 
 					return BundleAOrder < BundleBOrder;
 				});
+			}
 
-				// Sort remaining files into bundles
+			//Ensure all known sections are appended to the BulkBuildBundles Map even if they have no files
+			for (const FString& Section : SectionNames)
+			{
+				const FString BundleName = Section.RightChop(InstallBundleUtil::GetInstallBundleSectionPrefix().Len());
+				TArray<FString>& BundleFileList = BulkBuildBundles.FindOrAdd(FName(*BundleName));
+			}
+
+			//If we had files remaining to manually parse, now sort the remaining files into bundles
+			if (!bDidLoadAllFilesFromMetadata)
+			{
 				for (const FString& Section : SectionNames)
 				{
 					const FString BundleName = Section.RightChop(InstallBundleUtil::GetInstallBundleSectionPrefix().Len());
@@ -390,7 +400,7 @@ void FInstallBundleSourceBulk::AsyncInit_MakeBundlesForBulkBuild()
 		}
 
 		//See if we should serialize out the manually parsed results
-		if (!bDidLoadAllIniEntries)
+		if (!bDidLoadAllFilesFromMetadata)
 		{
 			bool bShouldSerializeMissingBulkBuildDataIni = false;
 			if (!GConfig->GetBool(TEXT("InstallBundleSource.Bulk.MiscSettings"), TEXT("bShouldSerializeMissingBulkBuildDataIni"), bShouldSerializeMissingBulkBuildDataIni, GInstallBundleIni))
