Fixed a bunch of non unity errors based on farm compile results

diff --git a/Engine/Plugins/Developer/TraceSourceFiltering/Source/SourceFilteringTrace/Private/TraceWorldFiltering.cpp b/Engine/Plugins/Developer/TraceSourceFiltering/Source/SourceFilteringTrace/Private/TraceWorldFiltering.cpp
--- a/Engine/Plugins/Developer/TraceSourceFiltering/Source/SourceFilteringTrace/Private/TraceWorldFiltering.cpp
+++ b/Engine/Plugins/Developer/TraceSourceFiltering/Source/SourceFilteringTrace/Private/TraceWorldFiltering.cpp
@@ -9,6 +9,8 @@
 
 #if WITH_EDITOR
 #include "Editor.h"
+#else
+#include "UObject/Package.h"
 #endif // WITH_EDITOR
 
 FDelegateHandle FTraceWorldFiltering::WorldInitHandle;
diff --git a/Engine/Plugins/Enterprise/VariantManagerContent/Source/VariantManagerContent/Private/ThumbnailGenerator.cpp b/Engine/Plugins/Enterprise/VariantManagerContent/Source/VariantManagerContent/Private/ThumbnailGenerator.cpp
--- a/Engine/Plugins/Enterprise/VariantManagerContent/Source/VariantManagerContent/Private/ThumbnailGenerator.cpp
+++ b/Engine/Plugins/Enterprise/VariantManagerContent/Source/VariantManagerContent/Private/ThumbnailGenerator.cpp
@@ -22,6 +22,8 @@
 #include "RendererInterface.h"
 #include "TextureCompiler.h"
 #include "ObjectTools.h"
+#else
+#include "Engine/World.h"
 #endif
 
 namespace ThumbnailGeneratorImpl
diff --git a/Engine/Plugins/Experimental/DefaultInstallBundleManager/Source/Private/InstallBundleSourcePlatformBase.cpp b/Engine/Plugins/Experimental/DefaultInstallBundleManager/Source/Private/InstallBundleSourcePlatformBase.cpp
--- a/Engine/Plugins/Experimental/DefaultInstallBundleManager/Source/Private/InstallBundleSourcePlatformBase.cpp
+++ b/Engine/Plugins/Experimental/DefaultInstallBundleManager/Source/Private/InstallBundleSourcePlatformBase.cpp
@@ -4,6 +4,8 @@
 
 
 #if WITH_PLATFORM_INSTALL_BUNDLE_SOURCE
+#include "DefaultInstallBundleManagerPrivate.h"
+#include "InstallBundleManagerPrivate.h"
 #include "InstallBundleManagerUtil.h"
 
 #include "Misc/ConfigCacheIni.h"
diff --git a/Engine/Plugins/Experimental/PCG/Source/PCG/Private/PCGHelpers.cpp b/Engine/Plugins/Experimental/PCG/Source/PCG/Private/PCGHelpers.cpp
--- a/Engine/Plugins/Experimental/PCG/Source/PCG/Private/PCGHelpers.cpp
+++ b/Engine/Plugins/Experimental/PCG/Source/PCG/Private/PCGHelpers.cpp
@@ -12,6 +12,8 @@
 
 #if WITH_EDITOR
 #include "Editor.h"
+#else
+#include "Engine/World.h"
 #endif
 
 namespace PCGHelpers
diff --git a/Engine/Plugins/Experimental/PCG/Source/PCG/Private/PCGSubsystem.cpp b/Engine/Plugins/Experimental/PCG/Source/PCG/Private/PCGSubsystem.cpp
--- a/Engine/Plugins/Experimental/PCG/Source/PCG/Private/PCGSubsystem.cpp
+++ b/Engine/Plugins/Experimental/PCG/Source/PCG/Private/PCGSubsystem.cpp
@@ -24,6 +24,7 @@
 #include "WorldPartition/ActorPartition/PartitionActorDesc.h"
 #include "ObjectTools.h"
 #else
+#include "Engine/Engine.h"
 #include "Engine/World.h"
 #endif
 
diff --git a/Engine/Plugins/Runtime/NetworkPrediction/Source/NetworkPrediction/Private/NetworkPredictionModule.cpp b/Engine/Plugins/Runtime/NetworkPrediction/Source/NetworkPrediction/Private/NetworkPredictionModule.cpp
--- a/Engine/Plugins/Runtime/NetworkPrediction/Source/NetworkPrediction/Private/NetworkPredictionModule.cpp
+++ b/Engine/Plugins/Runtime/NetworkPrediction/Source/NetworkPrediction/Private/NetworkPredictionModule.cpp
@@ -9,6 +9,8 @@
 #if WITH_EDITOR
 #include "Editor.h"
 #include "ISettingsModule.h"
+#else
+#include "Engine/World.h"
 #endif
 
 
