[DefaultInstallBundleManager]
* InstallBundleSourcePlatformBase - Removed bad include

diff --git a/Engine/Plugins/Experimental/DefaultInstallBundleManager/Source/Private/InstallBundleSourcePlatformBase.cpp b/Engine/Plugins/Experimental/DefaultInstallBundleManager/Source/Private/InstallBundleSourcePlatformBase.cpp
--- a/Engine/Plugins/Experimental/DefaultInstallBundleManager/Source/Private/InstallBundleSourcePlatformBase.cpp
+++ b/Engine/Plugins/Experimental/DefaultInstallBundleManager/Source/Private/InstallBundleSourcePlatformBase.cpp
@@ -5,7 +5,6 @@
 
 #if WITH_PLATFORM_INSTALL_BUNDLE_SOURCE
 #include "DefaultInstallBundleManagerPrivate.h"
-#include "InstallBundleManagerPrivate.h"
 #include "InstallBundleManagerUtil.h"
 
 #include "Misc/ConfigCacheIni.h"
