Dump relevant cache info on a failed cache reserve
Fix GFP state lex to string macro

diff --git a/Engine/Plugins/Experimental/DefaultInstallBundleManager/Source/Private/DefaultInstallBundleManager.cpp b/Engine/Plugins/Experimental/DefaultInstallBundleManager/Source/Private/DefaultInstallBundleManager.cpp
--- a/Engine/Plugins/Experimental/DefaultInstallBundleManager/Source/Private/DefaultInstallBundleManager.cpp
+++ b/Engine/Plugins/Experimental/DefaultInstallBundleManager/Source/Private/DefaultInstallBundleManager.cpp
@@ -919,13 +919,28 @@ void FDefaultInstallBundleManager::TryReserveCache(FContentRequestRef Request)
 
 	if (!bSuccess)
 	{
-		// Release from any caches that were reserved
+		LOG_INSTALL_BUNDLE_MAN_OVERRIDE(Request->LogVerbosityOverride, Display, TEXT("Failed to reserve cache for Request %s"), *Request->BundleName.ToString());
+
 		for (const TPair<FName, EInstallBundleCacheReserveResult>& Pair : ReserveResults)
 		{
 			if (Pair.Value == EInstallBundleCacheReserveResult::Success)
 			{
+				// Release from any caches that were reserved
 				verify(BundleCaches.FindChecked(Pair.Key)->Release(Request->BundleName));
 			}
+			else if (Pair.Value == EInstallBundleCacheReserveResult::Fail_CacheFull)
+			{
+				// Dump useful info
+				GetCacheStats(FInstallBundleSourceOrCache(Pair.Key), EInstallBundleCacheDumpToLog::Default, Request->GetLogVerbosityOverride());
+
+				TSharedRef<FInstallBundleCache> Cache = BundleCaches.FindChecked(Pair.Key);
+				if (TOptional<FInstallBundleCacheBundleInfo> CacheBundleInfo = Cache->GetBundleInfo(Request->BundleName))
+				{
+					LOG_INSTALL_BUNDLE_MAN_OVERRIDE(Request->LogVerbosityOverride, Display, TEXT("* Reserve attempt for request %s"), *Request->BundleName.ToString());
+					LOG_INSTALL_BUNDLE_MAN_OVERRIDE(Request->LogVerbosityOverride, Display, TEXT("* \tfull size: %" UINT64_FMT), CacheBundleInfo->FullInstallSize)
+					LOG_INSTALL_BUNDLE_MAN_OVERRIDE(Request->LogVerbosityOverride, Display, TEXT("* \tcurrent size: %" UINT64_FMT), CacheBundleInfo->CurrentInstallSize)
+				}
+			}
 		}
 
 		Request->LastCacheReserveResult = EInstallBundleCacheReserveResult::Fail_CacheFull;
diff --git a/Engine/Plugins/Experimental/GameFeatures/Source/GameFeatures/Private/GameFeaturesSubsystem.cpp b/Engine/Plugins/Experimental/GameFeatures/Source/GameFeatures/Private/GameFeaturesSubsystem.cpp
--- a/Engine/Plugins/Experimental/GameFeatures/Source/GameFeatures/Private/GameFeaturesSubsystem.cpp
+++ b/Engine/Plugins/Experimental/GameFeatures/Source/GameFeatures/Private/GameFeaturesSubsystem.cpp
@@ -89,7 +89,7 @@ namespace UE::GameFeatures
 #define GAME_FEATURE_PLUGIN_STATE_LEX_TO_STRING(inEnum, inText)  \
     case(EGameFeaturePluginState::inEnum):                       \
     {                                                            \
-        return TEXT("#inEnum");                                  \
+        return TEXT(#inEnum);                                    \
     }                                                            
 
 namespace GameFeaturePluginStatePrivate
diff --git a/Engine/Source/Runtime/InstallBundleManager/Private/InstallBundleCache.cpp b/Engine/Source/Runtime/InstallBundleManager/Private/InstallBundleCache.cpp
--- a/Engine/Source/Runtime/InstallBundleManager/Private/InstallBundleCache.cpp
+++ b/Engine/Source/Runtime/InstallBundleManager/Private/InstallBundleCache.cpp
@@ -46,10 +46,25 @@ void FInstallBundleCache::RemoveBundle(EInstallBundleSourceType Source, FName Bu
 	}
 }
 
-TOptional<FInstallBundleCacheBundleInfo> FInstallBundleCache::GetBundleInfo(EInstallBundleSourceType Source, FName BundleName) const
+TOptional<FInstallBundleCacheBundleInfo> FInstallBundleCache::GetBundleInfo(FName BundleName) const
 {
-	CSV_SCOPED_TIMING_STAT(InstallBundleManager, FInstallBundleCache_GetBundleInfo);
+	TOptional<FInstallBundleCacheBundleInfo> Ret;
+
+	if (FBundleCacheInfo* Info = CacheInfo.Find(BundleName))
+	{
+		FInstallBundleCacheBundleInfo& OutInfo = Ret.Emplace();
+		OutInfo.BundleName = BundleName;
+		OutInfo.FullInstallSize = Info->FullInstallSize;
+		OutInfo.CurrentInstallSize = Info->CurrentInstallSize;
+		OutInfo.TimeStamp = Info->TimeStamp;
+		OutInfo.AgeScalar = Info->AgeScalar;
+	}
 
+	return Ret;
+}
+
+TOptional<FInstallBundleCacheBundleInfo> FInstallBundleCache::GetBundleInfo(EInstallBundleSourceType Source, FName BundleName) const
+{
 	TOptional<FInstallBundleCacheBundleInfo> Ret;
 
 	const TMap<EInstallBundleSourceType, FPerSourceBundleCacheInfo>* SourcesMap = PerSourceCacheInfo.Find(BundleName);
diff --git a/Engine/Source/Runtime/InstallBundleManager/Public/InstallBundleCache.h b/Engine/Source/Runtime/InstallBundleManager/Public/InstallBundleCache.h
--- a/Engine/Source/Runtime/InstallBundleManager/Public/InstallBundleCache.h
+++ b/Engine/Source/Runtime/InstallBundleManager/Public/InstallBundleCache.h
@@ -63,7 +63,8 @@ public:
 
 	void RemoveBundle(EInstallBundleSourceType Source, FName BundleName);
 
-	TOptional<FInstallBundleCacheBundleInfo> GetBundleInfo(EInstallBundleSourceType Source, FName BundleName) const;
+	TOptional<FInstallBundleCacheBundleInfo> GetBundleInfo(FName BundleName) const;
+	TOptional<FInstallBundleCacheBundleInfo> GetBundleInfo(EInstallBundleSourceType Source, FName BundleName) const;	
 
 	// Return the total size of the cache
 	uint64 GetSize() const;
