Performance markers for content beacon and project edit compoenent and content service.

diff --git a/Engine/Plugins/Experimental/DefaultInstallBundleManager/Source/Private/DefaultInstallBundleManager.cpp b/Engine/Plugins/Experimental/DefaultInstallBundleManager/Source/Private/DefaultInstallBundleManager.cpp
--- a/Engine/Plugins/Experimental/DefaultInstallBundleManager/Source/Private/DefaultInstallBundleManager.cpp
+++ b/Engine/Plugins/Experimental/DefaultInstallBundleManager/Source/Private/DefaultInstallBundleManager.cpp
@@ -1413,6 +1413,7 @@ void FDefaultInstallBundleManager::TickAsyncMountTasks()
 void FDefaultInstallBundleManager::MountPaks(FContentRequestRef Request)
 {
 	SCOPED_BOOT_TIMING("FDefaultInstallBundleManager::MountPaks");
+	TRACE_BOOKMARK(TEXT("Start Mount Bundle %s"), *Request->BundleName.ToString());
 
 	check(!Request->bIsCanceled && Request->Result == EInstallBundleResult::OK);
 
@@ -1460,6 +1461,7 @@ void FDefaultInstallBundleManager::MountPaks(FContentRequestRef Request)
 		SetBundleStatus(BundleInfoLambda, EBundleState::Mounted);
 
 		StatsEnd(Request->BundleName, EContentRequestState::Mounting);
+		TRACE_BOOKMARK(TEXT("Finsihed Mount Bundle %s"), *Request->BundleName.ToString());
 		Request->StepResult = EContentRequestStepResult::Done;
 	};
 
