Add GetBundleSource to IInstallBundleManager

diff --git a/Engine/Plugins/Experimental/DefaultInstallBundleManager/Source/Private/DefaultInstallBundleManager.cpp b/Engine/Plugins/Experimental/DefaultInstallBundleManager/Source/Private/DefaultInstallBundleManager.cpp
--- a/Engine/Plugins/Experimental/DefaultInstallBundleManager/Source/Private/DefaultInstallBundleManager.cpp
+++ b/Engine/Plugins/Experimental/DefaultInstallBundleManager/Source/Private/DefaultInstallBundleManager.cpp
@@ -2922,6 +2922,15 @@ bool FDefaultInstallBundleManager::HasBundleSource(EInstallBundleSourceType Sour
 	return BundleSources.Contains(SourceType);
 }
 
+const TSharedPtr<IInstallBundleSource> FDefaultInstallBundleManager::GetBundleSource(EInstallBundleSourceType SourceType) const
+{
+	if (InitState != EInstallBundleManagerInitState::Succeeded)
+		return nullptr;
+
+	const TSharedPtr<IInstallBundleSource>* Found = BundleSources.Find(SourceType);
+	return Found ? *Found : nullptr;
+}
+
 FDelegateHandle FDefaultInstallBundleManager::PushInitErrorCallback(FInstallBundleManagerInitErrorHandler Callback)
 {
 	InitErrorHandlerStack.Push(MoveTemp(Callback));
diff --git a/Engine/Plugins/Experimental/DefaultInstallBundleManager/Source/Public/DefaultInstallBundleManager.h b/Engine/Plugins/Experimental/DefaultInstallBundleManager/Source/Public/DefaultInstallBundleManager.h
--- a/Engine/Plugins/Experimental/DefaultInstallBundleManager/Source/Public/DefaultInstallBundleManager.h
+++ b/Engine/Plugins/Experimental/DefaultInstallBundleManager/Source/Public/DefaultInstallBundleManager.h
@@ -605,6 +605,9 @@ protected:
 public:
 	virtual bool HasBundleSource(EInstallBundleSourceType SourceType) const override;
 
+	virtual const TSharedPtr<IInstallBundleSource> GetBundleSource(EInstallBundleSourceType SourceType) const override;
+
+
 	virtual FDelegateHandle PushInitErrorCallback(FInstallBundleManagerInitErrorHandler Callback) override;
 	virtual void PopInitErrorCallback() override;
 	virtual void PopInitErrorCallback(FDelegateHandle Handle) override;
diff --git a/Engine/Source/Runtime/InstallBundleManager/Private/InstallBundleManagerInterface.cpp b/Engine/Source/Runtime/InstallBundleManager/Private/InstallBundleManagerInterface.cpp
--- a/Engine/Source/Runtime/InstallBundleManager/Private/InstallBundleManagerInterface.cpp
+++ b/Engine/Source/Runtime/InstallBundleManager/Private/InstallBundleManagerInterface.cpp
@@ -48,6 +48,12 @@ TSharedPtr<IInstallBundleManager> IInstallBundleManager::GetPlatformInstallBundl
 	return {};
 }
 
+const TSharedPtr<IInstallBundleSource> IInstallBundleManager::GetBundleSource(EInstallBundleSourceType SourceType) const
+{
+	return {};
+}
+
+
 TValueOrError<FInstallBundleRequestInfo, EInstallBundleResult> IInstallBundleManager::RequestUpdateContent(FName BundleName, EInstallBundleRequestFlags Flags, ELogVerbosity::Type LogVerbosityOverride /*= ELogVerbosity::NoLogging*/)
 {
 	return RequestUpdateContent(MakeArrayView(&BundleName, 1), Flags, LogVerbosityOverride);
diff --git a/Engine/Source/Runtime/InstallBundleManager/Public/InstallBundleManagerInterface.h b/Engine/Source/Runtime/InstallBundleManager/Public/InstallBundleManagerInterface.h
--- a/Engine/Source/Runtime/InstallBundleManager/Public/InstallBundleManagerInterface.h
+++ b/Engine/Source/Runtime/InstallBundleManager/Public/InstallBundleManagerInterface.h
@@ -22,6 +22,7 @@
 #include "UObject/UnrealNames.h"
 
 class IAnalyticsProviderET;
+class IInstallBundleSource;
 
 struct FInstallBundleProgress
 {
@@ -102,6 +103,9 @@ public:
 
 	virtual bool HasBundleSource(EInstallBundleSourceType SourceType) const = 0;
 
+	virtual const TSharedPtr<IInstallBundleSource> GetBundleSource(EInstallBundleSourceType SourceType) const;
+
+
 	virtual FDelegateHandle PushInitErrorCallback(FInstallBundleManagerInitErrorHandler Callback) = 0;
 	virtual void PopInitErrorCallback() = 0;
 	virtual void PopInitErrorCallback(FDelegateHandle Handle) = 0;
