Fix FInstallBundleSourceBulk::GetBundleInstallState() being private and not protected.

diff --git a/Engine/Plugins/Experimental/DefaultInstallBundleManager/Source/Public/InstallBundleSourceBulk.h b/Engine/Plugins/Experimental/DefaultInstallBundleManager/Source/Public/InstallBundleSourceBulk.h
--- a/Engine/Plugins/Experimental/DefaultInstallBundleManager/Source/Public/InstallBundleSourceBulk.h
+++ b/Engine/Plugins/Experimental/DefaultInstallBundleManager/Source/Public/InstallBundleSourceBulk.h
@@ -22,6 +22,7 @@ private:
 	void AsyncInit_FireInitAnlaytic();
 	void AsyncInit_MakeBundlesForBulkBuild();
 
+protected:
 	virtual EInstallBundleInstallState GetBundleInstallState(FName BundleName);
 
 	// IInstallBundleSource Interface
