Add "ExcludedBundleSources" option to InstallBundles so that they can be skipped from certain sources.

diff --git a/Engine/Plugins/Experimental/DefaultInstallBundleManager/Source/Private/InstallBundleManagerUtil.cpp b/Engine/Plugins/Experimental/DefaultInstallBundleManager/Source/Private/InstallBundleManagerUtil.cpp
--- a/Engine/Plugins/Experimental/DefaultInstallBundleManager/Source/Private/InstallBundleManagerUtil.cpp
+++ b/Engine/Plugins/Experimental/DefaultInstallBundleManager/Source/Private/InstallBundleManagerUtil.cpp
@@ -82,6 +82,11 @@ namespace InstallBundleManagerUtil
 		OutInfo.BundleName = FName(*BundleName);
 		OutInfo.BundleNameString = BundleName;
 
+		TArray<FString> ExcludedBundleSources;
+		InstallBundleConfig.GetArray(*Section, TEXT("ExcludedBundleSources"), ExcludedBundleSources);
+		if (ExcludedBundleSources.Contains(LexToString(SourceType)))
+			return false;
+
 		if (!InstallBundleConfig.GetBool(*Section, TEXT("IsStartup"), OutInfo.bIsStartup))
 		{
 			OutInfo.bIsStartup = false;
