Update comment

diff --git a/Engine/Plugins/Experimental/DefaultInstallBundleManager/Source/Public/InstallBundleManagerUtil.h b/Engine/Plugins/Experimental/DefaultInstallBundleManager/Source/Public/InstallBundleManagerUtil.h
--- a/Engine/Plugins/Experimental/DefaultInstallBundleManager/Source/Public/InstallBundleManagerUtil.h
+++ b/Engine/Plugins/Experimental/DefaultInstallBundleManager/Source/Public/InstallBundleManagerUtil.h
@@ -358,7 +358,7 @@ namespace InstallBundleManagerAnalytics
 	 * @EventParam BundleName (string) Name of bundle.
 	 * @EventParam BundleSource (string) Source the bundle's content was removed from
 	 * @EventParam LastAccessTime (string) Format yyyy.mm.dd-hh.mm.ss, last access time for the BundleSource specified.
-	 * @EventParam BundleAgeHours (double) Amount of time in hours bundle content has been installed.
+	 * @EventParam BundleAgeHours (double) Amount of time in hours bundle content has been installed, measured from LastAccessTime.
 	 * @EventParam Result (string) Result of the request, OK or an error code
 	 */
 	DEFAULTINSTALLBUNDLEMANAGER_API void FireEvent_BundleEvictedFromCache(IAnalyticsProviderET* AnalyticsProvider,
