Add helper to suppress install bundle analytics to avoid spamming during locker randomization which can trigger a lot of installs per action and can be triggered quickly by users.

#tests Compiled ran game

diff --git a/Engine/Plugins/Experimental/DefaultInstallBundleManager/Source/Private/InstallBundleManagerUtil.cpp b/Engine/Plugins/Experimental/DefaultInstallBundleManager/Source/Private/InstallBundleManagerUtil.cpp
--- a/Engine/Plugins/Experimental/DefaultInstallBundleManager/Source/Private/InstallBundleManagerUtil.cpp
+++ b/Engine/Plugins/Experimental/DefaultInstallBundleManager/Source/Private/InstallBundleManagerUtil.cpp
@@ -605,14 +605,11 @@ namespace InstallBundleManagerUtil
 
 namespace InstallBundleManagerAnalytics
 {
-
-
-
 	void FireEvent_InitBundleManagerComplete(IAnalyticsProviderET* AnalyticsProvider, 
 		const bool bCanRetry, 
 		const FString InitResultString)
 	{
-		if (AnalyticsProvider == nullptr)
+		if (AnalyticsProvider == nullptr || InstallBundleUtil::FInstallBundleSuppressAnalytics::AllowAnalytics())
 		{
 			return;
 		}
@@ -624,7 +621,7 @@ namespace InstallBundleManagerAnalytics
 
 	void FireEvent_BundleManagerCacheStats(IAnalyticsProviderET* AnalyticsProvider, const FInstallBundleCacheStats& Stats)
 	{
-		if (AnalyticsProvider == nullptr)
+		if (AnalyticsProvider == nullptr || InstallBundleUtil::FInstallBundleSuppressAnalytics::AllowAnalytics())
 		{
 			return;
 		}
@@ -639,7 +636,7 @@ namespace InstallBundleManagerAnalytics
 
 	void FireEvent_InitBundleSourceBulkComplete(IAnalyticsProviderET* AnalyticsProvider, const FString InitResultString)
 	{
-		if (AnalyticsProvider == nullptr)
+		if (AnalyticsProvider == nullptr || InstallBundleUtil::FInstallBundleSuppressAnalytics::AllowAnalytics())
 		{
 			return;
 		}
@@ -650,7 +647,7 @@ namespace InstallBundleManagerAnalytics
 
 	void FireEvent_InitBundleSourcePlayGoComplete(IAnalyticsProviderET* AnalyticsProvider, const FString InitResultString)
 	{
-		if (AnalyticsProvider == nullptr)
+		if (AnalyticsProvider == nullptr || InstallBundleUtil::FInstallBundleSuppressAnalytics::AllowAnalytics())
 		{
 			return;
 		}
@@ -661,7 +658,7 @@ namespace InstallBundleManagerAnalytics
 
 	void FireEvent_InitBundleSourceIntelligentDeliveryComplete(IAnalyticsProviderET* AnalyticsProvider, const FString InitResultString)
 	{
-		if (AnalyticsProvider == nullptr)
+		if (AnalyticsProvider == nullptr || InstallBundleUtil::FInstallBundleSuppressAnalytics::AllowAnalytics())
 		{
 			return;
 		}
@@ -683,7 +680,7 @@ namespace InstallBundleManagerAnalytics
 	// Not deleting the code as usual because we want to maybe enable it later. Since we can't hotfix this code,
 	// we are forced to compile it out, but want a quick pathway to restore it if it comes back up that it is needed.
 #if 0
-		if (AnalyticsProvider == nullptr)
+		if (AnalyticsProvider == nullptr || InstallBundleUtil::FInstallBundleSuppressAnalytics::AllowAnalytics())
 		{
 			return;
 		}
@@ -703,7 +700,7 @@ namespace InstallBundleManagerAnalytics
 
 	void FireEvent_BundleRequestStarted(IAnalyticsProviderET* AnalyticsProvider, const FString& BundleName)
 	{
-		if (AnalyticsProvider == nullptr)
+		if (AnalyticsProvider == nullptr || InstallBundleUtil::FInstallBundleSuppressAnalytics::AllowAnalytics())
 		{
 			return;
 		}
@@ -718,7 +715,7 @@ namespace InstallBundleManagerAnalytics
 		const FString& Result, 
 		const InstallBundleUtil::FContentRequestStats& TimingStats)
 	{
-		if (AnalyticsProvider == nullptr)
+		if (AnalyticsProvider == nullptr || InstallBundleUtil::FInstallBundleSuppressAnalytics::AllowAnalytics())
 		{
 			return;
 		}
@@ -741,7 +738,7 @@ namespace InstallBundleManagerAnalytics
 
 	void FireEvent_BundleReleaseRequestStarted(IAnalyticsProviderET* AnalyticsProvider, const FString& BundleName, bool bRemoveFilesIfPossible)
 	{
-		if (AnalyticsProvider == nullptr)
+		if (AnalyticsProvider == nullptr || InstallBundleUtil::FInstallBundleSuppressAnalytics::AllowAnalytics())
 		{
 			return;
 		}
@@ -756,7 +753,7 @@ namespace InstallBundleManagerAnalytics
 		bool bRemoveFilesIfPossible, 
 		const FString& Result)
 	{
-		if (AnalyticsProvider == nullptr)
+		if (AnalyticsProvider == nullptr || InstallBundleUtil::FInstallBundleSuppressAnalytics::AllowAnalytics())
 		{
 			return;
 		}
@@ -769,7 +766,7 @@ namespace InstallBundleManagerAnalytics
 
 	void FireEvent_BundleEvictedFromCache(IAnalyticsProviderET* AnalyticsProvider, const FString& BundleName, const FString& BundleSource, FDateTime LastAccessTime, const FString& Result)
 	{
-		if (AnalyticsProvider == nullptr)
+		if (AnalyticsProvider == nullptr || InstallBundleUtil::FInstallBundleSuppressAnalytics::AllowAnalytics())
 		{
 			return;
 		}
@@ -786,7 +783,7 @@ namespace InstallBundleManagerAnalytics
 
 	void FireEvent_BundleCacheHit(IAnalyticsProviderET* AnalyticsProvider, const FString& BundleName, const FString& BundleSource)
 	{
-		if (AnalyticsProvider == nullptr)
+		if (AnalyticsProvider == nullptr || InstallBundleUtil::FInstallBundleSuppressAnalytics::AllowAnalytics())
 		{
 			return;
 		}
@@ -798,7 +795,7 @@ namespace InstallBundleManagerAnalytics
 
 	void FireEvent_BundleCacheMiss(IAnalyticsProviderET* AnalyticsProvider, const FString& BundleName, const FString& BundleSource, bool bPatchRequired)
 	{
-		if (AnalyticsProvider == nullptr)
+		if (AnalyticsProvider == nullptr || InstallBundleUtil::FInstallBundleSuppressAnalytics::AllowAnalytics())
 		{
 			return;
 		}
@@ -811,7 +808,7 @@ namespace InstallBundleManagerAnalytics
 
 	void FireEvent_PersistentPatchStats_StartPatching(IAnalyticsProviderET* AnalyticsProvider, const InstallBundleManagerUtil::FPersistentStatContainer::FPersistentStatsInformation& PersistentStatInformation)
 	{
-		if (AnalyticsProvider == nullptr)
+		if (AnalyticsProvider == nullptr || InstallBundleUtil::FInstallBundleSuppressAnalytics::AllowAnalytics())
 		{
 			return;
 		}
@@ -824,7 +821,7 @@ namespace InstallBundleManagerAnalytics
 
 	void FireEvent_PersistentPatchStats_EndPatching(IAnalyticsProviderET* AnalyticsProvider, const InstallBundleManagerUtil::FPersistentStatContainer::FPersistentStatsInformation& PersistentStatInformation)
 	{
-		if (AnalyticsProvider == nullptr)
+		if (AnalyticsProvider == nullptr || InstallBundleUtil::FInstallBundleSuppressAnalytics::AllowAnalytics())
 		{
 			return;
 		}
@@ -837,7 +834,7 @@ namespace InstallBundleManagerAnalytics
 
 	void FireEvent_PersistentPatchStats_Background(IAnalyticsProviderET* AnalyticsProvider, const InstallBundleManagerUtil::FPersistentStatContainer::FPersistentStatsInformation& PersistentStatInformation)
 	{
-		if (AnalyticsProvider == nullptr)
+		if (AnalyticsProvider == nullptr || InstallBundleUtil::FInstallBundleSuppressAnalytics::AllowAnalytics())
 		{
 			return;
 		}
@@ -850,7 +847,7 @@ namespace InstallBundleManagerAnalytics
 
 	void FireEvent_PersistentPatchStats_Foreground(IAnalyticsProviderET* AnalyticsProvider, const InstallBundleManagerUtil::FPersistentStatContainer::FPersistentStatsInformation& PersistentStatInformation)
 	{	
-		if (AnalyticsProvider == nullptr)
+		if (AnalyticsProvider == nullptr || InstallBundleUtil::FInstallBundleSuppressAnalytics::AllowAnalytics())
 		{
 			return;
 		}
diff --git a/Engine/Source/Runtime/InstallBundleManager/Private/InstallBundleUtils.cpp b/Engine/Source/Runtime/InstallBundleManager/Private/InstallBundleUtils.cpp
--- a/Engine/Source/Runtime/InstallBundleManager/Private/InstallBundleUtils.cpp
+++ b/Engine/Source/Runtime/InstallBundleManager/Private/InstallBundleUtils.cpp
@@ -78,6 +78,42 @@ namespace InstallBundleUtil
 		}
 	}
 
+	TAtomic<int32> InstallBundleSuppressAnalyticsCounter = 0;
+
+	FInstallBundleSuppressAnalytics::FInstallBundleSuppressAnalytics()
+		: bIsEnabled(false)
+	{
+	}
+
+	FInstallBundleSuppressAnalytics::~FInstallBundleSuppressAnalytics()
+	{
+		Disable();
+	}
+
+	void FInstallBundleSuppressAnalytics::Enable()
+	{
+		if (!bIsEnabled)
+		{
+			bIsEnabled = true;
+			ensure(InstallBundleSuppressAnalyticsCounter++ >= 0);
+		}
+	}
+
+	void FInstallBundleSuppressAnalytics::Disable()
+	{
+		if (bIsEnabled)
+		{
+			bIsEnabled = false;
+			ensure(--InstallBundleSuppressAnalyticsCounter >= 0);
+		}
+	}
+
+	bool FInstallBundleSuppressAnalytics::AllowAnalytics()
+	{
+		return InstallBundleSuppressAnalyticsCounter == 0;
+	}
+
+
 	void StartInstallBundleAsyncIOTask(TArray<TUniquePtr<FInstallBundleTask>>& Tasks, TUniqueFunction<void()> WorkFunc, TUniqueFunction<void()> OnComplete)
 	{
 		TUniquePtr<FInstallBundleTask> Task = MakeUnique<FInstallBundleTask>(MoveTemp(WorkFunc), MoveTemp(OnComplete));
diff --git a/Engine/Source/Runtime/InstallBundleManager/Public/InstallBundleUtils.h b/Engine/Source/Runtime/InstallBundleManager/Public/InstallBundleUtils.h
--- a/Engine/Source/Runtime/InstallBundleManager/Public/InstallBundleUtils.h
+++ b/Engine/Source/Runtime/InstallBundleManager/Public/InstallBundleUtils.h
@@ -123,6 +123,21 @@ namespace InstallBundleUtil
 		}
 	};
 
+	// Helper to safely suppress installer analytics that will restore previous state when going out of scope
+	struct INSTALLBUNDLEMANAGER_API FInstallBundleSuppressAnalytics : FNoncopyable
+	{
+		FInstallBundleSuppressAnalytics();
+		~FInstallBundleSuppressAnalytics();
+
+		void Enable();
+		void Disable();
+
+		static bool AllowAnalytics();
+
+	private:
+		bool bIsEnabled;
+	};
+
 	class INSTALLBUNDLEMANAGER_API FInstallBundleWork : public FNonAbandonableTask
 	{
 	public:
