Install Bundle Manager
- Use UE::Task::FPipe instead of a dedicate journal thread.

diff --git a/Engine/Plugins/Experimental/DefaultInstallBundleManager/Source/Public/InstallBundleManagerUtil.h b/Engine/Plugins/Experimental/DefaultInstallBundleManager/Source/Public/InstallBundleManagerUtil.h
--- a/Engine/Plugins/Experimental/DefaultInstallBundleManager/Source/Public/InstallBundleManagerUtil.h
+++ b/Engine/Plugins/Experimental/DefaultInstallBundleManager/Source/Public/InstallBundleManagerUtil.h
@@ -26,8 +26,8 @@ namespace InstallBundleManagerUtil
 	DEFAULTINSTALLBUNDLEMANAGER_API TSharedPtr<IInstallBundleSource> MakePlatformBundleSource();
 #endif
 
-
 	// Returns a thread pool with one thread suitible for running in-order journal tasks
+	UE_DEPRECATED(5.4, "Use UE::Tasks::FPipe instead.")
 	DEFAULTINSTALLBUNDLEMANAGER_API TSharedPtr<FQueuedThreadPool, ESPMode::ThreadSafe> GetJournalThreadPool();
 
 	// Fills out a FInstallBundleSourceBundleInfo from the specified config section
