Add a new experimental bundle source that uses the platform chunk installer, via some optional new functions in the chunk installer interface.

#rnx

diff --git a/Engine/Plugins/Experimental/DefaultInstallBundleManager/Source/Private/InstallBundleSourcePlatformChunkInstall.cpp b/Engine/Plugins/Experimental/DefaultInstallBundleManager/Source/Private/InstallBundleSourcePlatformChunkInstall.cpp
new file mode 100644
--- /dev/null
+++ b/Engine/Plugins/Experimental/DefaultInstallBundleManager/Source/Private/InstallBundleSourcePlatformChunkInstall.cpp
@@ -0,0 +1,667 @@
+// Copyright Epic Games, Inc. All Rights Reserved.
+
+#include "InstallBundleSourcePlatformChunkInstall.h"
+#if WITH_PLATFORM_INSTALL_BUNDLE_SOURCE
+#include "DefaultInstallBundleManagerPrivate.h"
+
+#include "InstallBundleManagerUtil.h"
+
+#include "Misc/CommandLine.h"
+#include "Misc/ConfigCacheIni.h"
+#include "Algo/AnyOf.h"
+#include "Algo/Find.h"
+#include "Containers/Ticker.h"
+#include "IPlatformFilePak.h"
+#include "Stats/Stats.h"
+
+#define LOG_SOURCE_CHUNKINSTALL(Verbosity, Format, ...) LOG_INSTALL_BUNDLE_MAN(Verbosity, TEXT("InstallBundleSourceIntelligentDelivery: ") Format, ##__VA_ARGS__)
+
+#define LOG_SOURCE_CHUNKINSTALL_OVERRIDE(VerbosityOverride, Verbosity, Format, ...) LOG_INSTALL_BUNDLE_MAN_OVERRIDE(VerbosityOverride, Verbosity, TEXT("InstallBundleSourceIntelligentDelivery: ") Format, ##__VA_ARGS__)
+
+
+
+FInstallBundleSourcePlatformChunkInstall::FInstallBundleSourcePlatformChunkInstall(IPlatformChunkInstall* InPlatformChunkInstall)
+	: PlatformChunkInstall(InPlatformChunkInstall)
+{
+	TickHandle = FTSTicker::GetCoreTicker().AddTicker(FTickerDelegate::CreateRaw(this, &FInstallBundleSourcePlatformChunkInstall::Tick));
+
+	NamedChunkInstallDelegateHandle = PlatformChunkInstall->AddNamedChunkInstallDelegate( FPlatformNamedChunkInstallDelegate::CreateRaw(this, &FInstallBundleSourcePlatformChunkInstall::OnNamedChunkInstall));
+
+	InPlatformChunkInstall->SetAutoPakMountingEnabled(false);
+}
+
+
+FInstallBundleSourcePlatformChunkInstall::~FInstallBundleSourcePlatformChunkInstall()
+{
+	PlatformChunkInstall->RemoveNamedChunkInstallDelegate(NamedChunkInstallDelegateHandle);
+	NamedChunkInstallDelegateHandle.Reset();
+
+	FTSTicker::GetCoreTicker().RemoveTicker(TickHandle);
+	TickHandle.Reset();
+
+	// Cleanup any Async tasks
+	InstallBundleUtil::CleanupInstallBundleAsyncIOTasks(GeneralAsyncTasks);
+}
+
+
+bool FInstallBundleSourcePlatformChunkInstall::Tick(float dt)
+{
+	QUICK_SCOPE_CYCLE_COUNTER(STAT_FInstallBundleSourcePlatformChunkInstall_Tick);
+
+	CSV_SCOPED_TIMING_STAT(InstallBundleManager, FInstallBundleSourcePlatformChunkInstall_Tick);
+
+	InstallBundleUtil::FinishInstallBundleAsyncIOTasks(GeneralAsyncTasks);
+
+	// remove finished requests
+	ContentReleaseRequests.RemoveAll([](const FContentReleaseRequestRef& Request) { return !Request->bInProgress; });
+	bool bHasRemovedRequests = ContentRequests.RemoveAll([](const FContentRequestRef& Request) { return !Request->bInProgress; }) > 0;
+	bChunkOrderDirty |= bHasRemovedRequests;
+
+	// make sure we're installing the highest priority chunks
+	TickUpdateChunkOrder();
+
+	return true;
+}
+
+
+void FInstallBundleSourcePlatformChunkInstall::TickUpdateChunkOrder()
+{
+	CSV_SCOPED_TIMING_STAT(InstallBundleManager, FInstallBundleSourcePlatformChunkInstall_TickUpdateChunkOrder);
+
+	if (bChunkOrderDirty && !bIsUpdatingChunkOrder)
+	{
+		bChunkOrderDirty = false; // immediately clear dirty flag as we may get dirty while updating
+		bIsUpdatingChunkOrder = true;
+
+		// find the highest priority group of currently-installing bundles
+		EInstallBundlePriority HighestPriority = EInstallBundlePriority::Count;
+		TArray<FName> HighestPriorityNamedChunks;
+
+		for (const FContentRequestRef& ContentRequest : ContentRequests)
+		{
+			const FBundleInfo* BundleInfo = BundleInfoMap.Find(ContentRequest->BundleName);
+			if (BundleInfo != nullptr)
+			{
+				if (BundleInfo->Priority < HighestPriority) // NB. lower value = more important!
+				{
+					HighestPriorityNamedChunks.Reset();
+					HighestPriority = BundleInfo->Priority;
+				}
+				if (BundleInfo->Priority == HighestPriority)
+				{
+					HighestPriorityNamedChunks.Add(BundleInfo->NamedChunk);
+				}
+			}
+		}
+
+		// update the feature priority asyncronously
+		if (HighestPriorityNamedChunks.Num() > 0)
+		{
+			InstallBundleUtil::StartInstallBundleAsyncIOTask(GeneralAsyncTasks, 
+				[this,PriorityNamedChunks = MoveTemp(HighestPriorityNamedChunks)]()
+				{
+					for (FName NamedChunk : PriorityNamedChunks)
+					{
+						PlatformChunkInstall->PrioritizeNamedChunk( NamedChunk, EChunkPriority::High );
+					}
+				},
+				[this]()
+				{
+					bIsUpdatingChunkOrder = false;
+				}
+			);
+		}
+		else
+		{
+			// nothing to update
+			bIsUpdatingChunkOrder = false;
+		}
+	}
+}
+
+
+FInstallBundleSourceInitInfo FInstallBundleSourcePlatformChunkInstall::Init(TSharedRef<InstallBundleUtil::FContentRequestStatsMap> InRequestStats, TSharedPtr<IAnalyticsProviderET> InAnalyticsProvider, TSharedPtr<InstallBundleUtil::PersistentStats::FPersistentStatContainerBase> InPersistentStatsContainer)
+{
+	CSV_SCOPED_TIMING_STAT(InstallBundleManager, FInstallBundleSourcePlatformChunkInstall_Init);
+	FInstallBundleSourceInitInfo InitInfo;
+
+	if (!PlatformChunkInstall->SupportsBundleSource())
+	{
+		LOG_SOURCE_CHUNKINSTALL(Display, TEXT("Platform chunk installer doesn't support bundles, attempting to fallback to next bundle source."));
+
+		InitState = EInstallBundleManagerInitState::Failed;
+
+		InitInfo.Result = EInstallBundleManagerInitResult::BuildMetaDataNotFound;
+		InitInfo.bShouldUseFallbackSource = true;
+		return InitInfo;
+	}
+
+	InitInfo = FInstallBundleSourcePlatformBase::Init(InRequestStats, InAnalyticsProvider, InPersistentStatsContainer);
+	return InitInfo;
+}
+
+
+void FInstallBundleSourcePlatformChunkInstall::AsyncInit(FInstallBundleSourceInitDelegate Callback)
+{
+	CSV_SCOPED_TIMING_STAT(InstallBundleManager, FInstallBundleSourcePlatformChunkInstall_AsyncInit);
+
+	// Handle retrying init if recoverable
+	if (InitState == EInstallBundleManagerInitState::Failed)
+	{
+		LOG_SOURCE_CHUNKINSTALL(Warning, TEXT("Retrying initialization"));
+		InitState = EInstallBundleManagerInitState::NotInitialized;
+	}
+	
+	struct AsyncInitContext
+	{
+		EInstallBundleManagerInitResult InitResult = EInstallBundleManagerInitResult::OK;
+
+		TSet<FName> NamedChunks;
+		TMap<FName, FBundleInfo> BundleInfoMap;
+		TArray<TPair<FString, TArray<FRegexPattern>>> BundleRegexList;
+		IPlatformChunkInstall* PlatformChunkInstall;
+	};
+
+
+	// create and initialize the async init context
+	TSharedPtr<AsyncInitContext, ESPMode::ThreadSafe> Context = MakeShared<AsyncInitContext, ESPMode::ThreadSafe>();
+	Context->BundleRegexList = MoveTemp(BundleRegexList);
+	Context->PlatformChunkInstall = PlatformChunkInstall;
+
+	
+
+	auto AsyncInit = [Context]() mutable
+	{
+		// make sure there's a config
+		const FConfigFile* InstallBundleConfig = GConfig->FindConfigFile(GInstallBundleIni);
+		if (!InstallBundleConfig)
+		{
+			Context->InitResult = EInstallBundleManagerInitResult::ConfigurationError;
+			return;
+		}
+
+		// cache all named chunks
+		Context->NamedChunks.Append( Context->PlatformChunkInstall->GetNamedChunksByType(ENamedChunkType::OnDemand) );
+		Context->NamedChunks.Append( Context->PlatformChunkInstall->GetNamedChunksByType(ENamedChunkType::Language) );
+
+		// create bundles & chunk data
+		TSet<FName> KnownNamedChunks;
+		TMap<FName, TArray<FString>> NamedChunkFilePaths;
+		FString RootDir = FPaths::RootDir();
+		bool bHasAbsoluteRootDir = !FPaths::IsRelative(RootDir);
+
+		for (const TPair<FString, FConfigSection>& Pair : *InstallBundleConfig)
+		{
+			// make sure this is a suitable platform bundle
+			const FString& Section = Pair.Key;
+			if (!Section.StartsWith(InstallBundleUtil::GetInstallBundleSectionPrefix()))
+			{
+				continue;
+			}
+
+			int32 PlatformChunkID = 0;
+			if (InstallBundleConfig->GetInt(*Section, TEXT("PlatformChunkID"), PlatformChunkID) && PlatformChunkID < 0)
+			{
+				continue;
+			}
+
+			// create bundle info
+			FName BundleName( *Section.RightChop(InstallBundleUtil::GetInstallBundleSectionPrefix().Len()));
+			FBundleInfo& BundleInfo = Context->BundleInfoMap.Add(BundleName);
+			BundleInfo.NamedChunk = BundleName;
+			BundleInfo.Priority = EInstallBundlePriority::Normal;
+
+			FString PriorityString;
+			InstallBundleConfig->GetString(*Section, TEXT("Priority"), PriorityString);
+			LexTryParseString(BundleInfo.Priority, *PriorityString);
+
+			if (KnownNamedChunks.Contains(BundleInfo.NamedChunk))
+			{
+				continue;
+			}
+			KnownNamedChunks.Add(BundleInfo.NamedChunk);
+
+			if (Context->NamedChunks.Contains(BundleInfo.NamedChunk))
+			{
+				// collect the paths for this bundle
+				TArray<FString> FilesInChunk;
+				if (Context->PlatformChunkInstall->GetPakFilesInNamedChunk(BundleInfo.NamedChunk, FilesInChunk))
+				{
+					for (const FString& FilePath : FilesInChunk)
+					{
+						FString ChunkFilePath = FilePath;
+						if (bHasAbsoluteRootDir)
+						{
+							ChunkFilePath = TEXT("../../../") + ChunkFilePath;
+						}
+						FPaths::NormalizeFilename(ChunkFilePath);
+						NamedChunkFilePaths.FindOrAdd(BundleName).Add(ChunkFilePath);
+					}
+				}
+			}
+		}
+
+		// collect pak search directories
+		TArray<FString> PakFolders;
+		FPakPlatformFile::GetPakFolders(FCommandLine::Get(), PakFolders);
+		
+		TArray<FString> PakSearchDirs;
+		PakSearchDirs.Empty(PakFolders.Num());
+		for (FString& PakFolder : PakFolders)
+		{
+			if (bHasAbsoluteRootDir && FPaths::IsRelative(PakFolder))
+			{
+				PakSearchDirs.Add_GetRef(MoveTemp(PakFolder));
+			}
+			else if (PakFolder.StartsWith(RootDir))
+			{
+				verify(FPaths::MakePathRelativeTo(PakSearchDirs.Add_GetRef(MoveTemp(PakFolder)), *RootDir));
+			}
+		}
+
+		// More than one bundle can map to the same chunk, so map chunk files to the correct bundles
+		for (TPair<FName, TArray<FString>>& FilePathPair : NamedChunkFilePaths)
+		{
+			for (FString& Path : FilePathPair.Value)
+			{
+				if (!Algo::AnyOf(PakSearchDirs, [&Path](const FString& Dir) { return Path.StartsWith(Dir); }))
+				{
+					// Only want to mount content in PakSearchDirs and don't mount
+					// any other content that may be in the package.
+					continue;
+				}
+		
+				FString BundleName;
+				if (InstallBundleUtil::MatchBundleRegex(Context->BundleRegexList, Path, BundleName))
+				{
+					FBundleInfo& BundleInfo = Context->BundleInfoMap.FindChecked(*BundleName);
+					BundleInfo.FilePaths.AddUnique(MoveTemp(Path));
+				}
+				else
+				{
+					checkf(false, TEXT("Failed to map chunk file %s to an install bundle"), *Path);
+				}
+			}
+		}
+	};
+
+	auto OnAsyncInitComplete = [this, Context, OnInitCompleteCallback = MoveTemp(Callback)]() mutable
+	{
+		check(IsInGameThread());
+
+		FInstallBundleSourceAsyncInitInfo InitInfo;
+		InitInfo.Result = Context->InitResult;
+
+		if (InitInfo.Result == EInstallBundleManagerInitResult::OK)
+		{
+			// initialization succeeded - store the initialized data
+			BundleInfoMap = MoveTemp(Context->BundleInfoMap);
+			NamedChunks = MoveTemp(Context->NamedChunks);
+
+			InitState = EInstallBundleManagerInitState::Succeeded;
+		}
+		else
+		{
+			LOG_SOURCE_CHUNKINSTALL(Error, TEXT("Initialization Failed - %s"), LexToString(InitInfo.Result));
+			InitState = EInstallBundleManagerInitState::Failed;
+		}
+
+		LOG_SOURCE_CHUNKINSTALL(Display, TEXT("Fire Init Analytic: %s"), LexToString(InitInfo.Result));
+		InstallBundleManagerAnalytics::FireEvent_InitBundleSourceIntelligentDeliveryComplete(AnalyticsProvider.Get(), LexToString(InitInfo.Result));
+
+		OnInitCompleteCallback.Execute(AsShared(), MoveTemp(InitInfo));
+	};
+
+	// start the async initialize
+	InstallBundleUtil::StartInstallBundleAsyncIOTask(GeneralAsyncTasks, MoveTemp(AsyncInit), MoveTemp(OnAsyncInitComplete));
+}
+
+
+bool FInstallBundleSourcePlatformChunkInstall::QueryPersistentBundleInfo(FInstallBundleSourcePersistentBundleInfo& SourceBundleInfo) const
+{
+	// lookup bundle info
+	const FBundleInfo* BundleInfo = BundleInfoMap.Find(SourceBundleInfo.BundleName);
+	if (BundleInfo == nullptr)
+	{
+		return false;
+	}
+
+	// assume the startup bundle is installed. It cannot be uninstalled
+	if (SourceBundleInfo.bIsStartup)
+	{
+		SourceBundleInfo.BundleContentState = EInstallBundleInstallState::UpToDate;
+		return true;
+	}
+
+	// lookup the chunk data
+	FChunkInstallationStatusDetail ChunkStatusDetail;
+	if (PlatformChunkInstall->GetNamedChunkInstallationStatus(BundleInfo->NamedChunk, ChunkStatusDetail))
+	{
+		SourceBundleInfo.CurrentInstallSize = ChunkStatusDetail.CurrentInstallSize;
+		SourceBundleInfo.FullInstallSize = ChunkStatusDetail.FullInstallSize;
+		SourceBundleInfo.BundleContentState = ChunkStatusDetail.bIsInstalled ? EInstallBundleInstallState::UpToDate : EInstallBundleInstallState::NotInstalled;
+		return true;
+	}
+
+	return false;
+}
+
+
+void FInstallBundleSourcePlatformChunkInstall::GetContentState(TArrayView<const FName> BundleNames, EInstallBundleGetContentStateFlags Flags, FInstallBundleGetContentStateDelegate Callback)
+{
+	CSV_SCOPED_TIMING_STAT(InstallBundleManager, FInstallBundleSourcePlatformChunkInstall_GetContentState);
+
+	FInstallBundleCombinedContentState State;
+	State.CurrentVersion.Add(GetSourceType(), InstallBundleUtil::GetAppVersion());
+
+	// capture the install state and version for all bundles, and record the remaining download size
+	uint64 TotalRemainingDownloadSize = 0;
+	TMap<FName,uint64> BundleToRemaingDownloadSize;
+	for (const FName& BundleName : BundleNames)
+	{
+		const FBundleInfo* BundleInfo = BundleInfoMap.Find(BundleName);
+		if (BundleInfo != nullptr)
+		{
+			FChunkInstallationStatusDetail ChunkStatusDetail;
+			if (PlatformChunkInstall->GetNamedChunkInstallationStatus( BundleInfo->NamedChunk, ChunkStatusDetail ))
+			{
+				FInstallBundleContentState& IndividualBundleState = State.IndividualBundleStates.Add(BundleName);
+
+				if (ChunkStatusDetail.bIsInstalled)
+				{
+					IndividualBundleState.State = EInstallBundleInstallState::UpToDate;
+					IndividualBundleState.Version.Add(GetSourceType(), InstallBundleUtil::GetAppVersion());
+				}
+				else
+				{
+					IndividualBundleState.State = EInstallBundleInstallState::NotInstalled;
+					IndividualBundleState.Version.Add(GetSourceType(), TEXT(""));
+				}
+
+				check( ChunkStatusDetail.CurrentInstallSize <= ChunkStatusDetail.FullInstallSize);
+				uint64 RemainingDownloadSize = (ChunkStatusDetail.FullInstallSize - ChunkStatusDetail.CurrentInstallSize);
+
+				TotalRemainingDownloadSize += RemainingDownloadSize;
+				BundleToRemaingDownloadSize.Add(BundleName, RemainingDownloadSize);
+			}
+		}
+	}
+
+	// compute the download weight for all of the bundles - higher weight is a bigger download
+	for (TTuple<FName,uint64> Pair : BundleToRemaingDownloadSize)
+	{
+		uint64 RemainingDownloadSize = Pair.Value;
+		double Weight = (TotalRemainingDownloadSize > 0)  ?  ((double)RemainingDownloadSize / (double)TotalRemainingDownloadSize)  :  0.0;
+
+		FInstallBundleContentState& IndividualBundleState = State.IndividualBundleStates.FindChecked(Pair.Key);
+		IndividualBundleState.Weight = FMath::Max(Weight, InstallBundleUtil::MinimumBundleWeight); // this Max() does mean the total weight will be > 1.0 but all other bundle sources do it this way too
+	}
+
+	// Don't return any size info since all the size should be reserved by the system	
+	Callback.ExecuteIfBound(MoveTemp(State));
+}
+
+
+EInstallBundleSourceBundleSkipReason FInstallBundleSourcePlatformChunkInstall::GetBundleSkipReason(FName BundleName) const
+{
+	EInstallBundleSourceBundleSkipReason SkipReason = EInstallBundleSourceBundleSkipReason::None;
+
+	const FBundleInfo* BundleInfo = BundleInfoMap.Find(BundleName);
+	if (BundleInfo == nullptr)
+	{
+		LOG_SOURCE_CHUNKINSTALL(Warning, TEXT("Skipping Bundle %s - no equivalent platform named chunk - allowing anyway"), *BundleName.ToString());
+		return SkipReason;
+	}
+
+	bool SkipLanguageBundlesIfNotForCurrentLocale = true;
+	if (!GConfig->GetBool(TEXT("InstallBundleSource.Platform.MiscSettings"), TEXT("SkipLanguageBundlesIfNotForCurrentLocale"), SkipLanguageBundlesIfNotForCurrentLocale, GInstallBundleIni))
+	{
+		SkipLanguageBundlesIfNotForCurrentLocale = true;
+	}
+	if (SkipLanguageBundlesIfNotForCurrentLocale && !PlatformChunkInstall->IsNamedChunkForCurrentLocale(BundleInfo->NamedChunk))
+	{
+		LOG_SOURCE_CHUNKINSTALL(Verbose, TEXT("Skipping Bundle %s - Locale not current"), *BundleName.ToString());
+		SkipReason |= EInstallBundleSourceBundleSkipReason::LanguageNotCurrent;
+	}
+
+	return SkipReason;
+}
+
+
+void FInstallBundleSourcePlatformChunkInstall::RequestUpdateContent(FRequestUpdateContentBundleContext Context)
+{
+	LOG_SOURCE_CHUNKINSTALL_OVERRIDE(Context.LogVerbosityOverride, Display, TEXT("Requesting Bundle %s"), *Context.BundleName.ToString());
+
+	// sanity check the request
+	bool bFailed = false;
+	const FBundleInfo* BundleInfo = BundleInfoMap.Find(Context.BundleName);
+	if (BundleInfo == nullptr)
+	{
+		LOG_SOURCE_CHUNKINSTALL_OVERRIDE(Context.LogVerbosityOverride, Display, TEXT("Bundle %s is unknown"), *Context.BundleName.ToString());
+		bFailed = true;
+	}
+	else if (!NamedChunks.Contains(BundleInfo->NamedChunk))
+	{
+		LOG_SOURCE_CHUNKINSTALL_OVERRIDE(Context.LogVerbosityOverride, Display, TEXT("Bundle %s is not named chunk"), *Context.BundleName.ToString());
+		bFailed = true;
+	}
+	else if (ContentRequests.ContainsByPredicate([BundleName = Context.BundleName](const FContentRequestRef& ContentRequest) { return ContentRequest->BundleName == BundleName; }))
+	{
+		LOG_SOURCE_CHUNKINSTALL_OVERRIDE(Context.LogVerbosityOverride, Error, TEXT("Bundle %s install already in progress (named chunk: %s)"), *Context.BundleName.ToString(), *BundleInfo->NamedChunk.ToString());
+		bFailed = true;
+	}
+	else if (PlatformChunkInstall->GetNamedChunkLocation(BundleInfo->NamedChunk) == EChunkLocation::LocalFast)
+	{
+		LOG_SOURCE_CHUNKINSTALL_OVERRIDE(Context.LogVerbosityOverride, Error, TEXT("Bundle %s is already installed (named chunk: %s)"), *Context.BundleName.ToString(), *BundleInfo->NamedChunk.ToString());
+		bFailed = true;
+	}
+	else
+	{
+		bool bIsInstalling = PlatformChunkInstall->InstallNamedChunk(BundleInfo->NamedChunk);
+		if (bIsInstalling)
+		{
+			FContentRequestRef& NewRequest = ContentRequests.Emplace_GetRef(MakeShared<FContentRequest, ESPMode::ThreadSafe>());
+			NewRequest->BundleName = Context.BundleName;
+			NewRequest->ContentPaths = BundleInfo->FilePaths;
+			NewRequest->LogVerbosityOverride = Context.LogVerbosityOverride;
+			NewRequest->CompleteCallback = MoveTemp(Context.CompleteCallback);
+
+			bChunkOrderDirty = true;
+		}
+		else
+		{
+			LOG_SOURCE_CHUNKINSTALL_OVERRIDE(Context.LogVerbosityOverride, Display, TEXT("Bundle %s failed to start"), *Context.BundleName.ToString());
+			bFailed = true;
+		}
+	}
+	
+	if (bFailed)
+	{
+		// send the failure callback immediately
+		FInstallBundleSourceUpdateContentResultInfo ResultInfo;
+		ResultInfo.BundleName = Context.BundleName;
+		ResultInfo.Result = EInstallBundleResult::OK;
+		Context.CompleteCallback.ExecuteIfBound(AsShared(), MoveTemp(ResultInfo));
+	}
+}
+
+
+void FInstallBundleSourcePlatformChunkInstall::RequestReleaseContent(FRequestReleaseContentBundleContext Context)
+{
+	LOG_SOURCE_CHUNKINSTALL_OVERRIDE(Context.LogVerbosityOverride, Display, TEXT("Requesting Release Bundle %s"), *Context.BundleName.ToString());
+
+	// sanity check the request
+	bool bFailed = false;
+	const FBundleInfo* BundleInfo = BundleInfoMap.Find(Context.BundleName);
+	if (BundleInfo == nullptr)
+	{
+		LOG_SOURCE_CHUNKINSTALL_OVERRIDE(Context.LogVerbosityOverride, Display, TEXT("Bundle %s is unknown"), *Context.BundleName.ToString());
+		bFailed = true;
+	}
+	else if (!NamedChunks.Contains(BundleInfo->NamedChunk))
+	{
+		LOG_SOURCE_CHUNKINSTALL_OVERRIDE(Context.LogVerbosityOverride, Display, TEXT("Bundle %s is not named chunk"), *Context.BundleName.ToString());
+		bFailed = true;
+	}
+	else if ( ContentReleaseRequests.ContainsByPredicate( [BundleName = Context.BundleName](const FContentReleaseRequestRef& ContentReleaseRequest) { return ContentReleaseRequest->BundleName == BundleName; } ))
+	{
+		LOG_SOURCE_CHUNKINSTALL_OVERRIDE(Context.LogVerbosityOverride, Error, TEXT("Bundle %s removal already in progress (named chunk: %s)"), *Context.BundleName.ToString(), *BundleInfo->NamedChunk.ToString());
+		bFailed = true;
+	}
+	else if (PlatformChunkInstall->GetNamedChunkLocation(BundleInfo->NamedChunk) == EChunkLocation::NotAvailable)
+	{
+		LOG_SOURCE_CHUNKINSTALL_OVERRIDE(Context.LogVerbosityOverride, Error, TEXT("Bundle %s is already uninstalled (named chunk: %s)"), *Context.BundleName.ToString(), *BundleInfo->NamedChunk.ToString());
+
+		FInstallBundleSourceReleaseContentResultInfo ResultInfo;
+		ResultInfo.BundleName = Context.BundleName;
+		ResultInfo.Result = EInstallBundleReleaseResult::OK;
+		ResultInfo.bContentWasRemoved = true;
+		Context.CompleteCallback.ExecuteIfBound(AsShared(), MoveTemp(ResultInfo));
+	}
+	else if (!EnumHasAnyFlags(Context.Flags, EInstallBundleReleaseRequestFlags::RemoveFilesIfPossible))
+	{
+		// caller does not want the content removed
+		bFailed = true;
+	}
+	else
+	{
+		bool bIsUninstalling = PlatformChunkInstall->UninstallNamedChunk(BundleInfo->NamedChunk);
+		if (bIsUninstalling)
+		{
+			FContentReleaseRequestRef& NewRequest = ContentReleaseRequests.Emplace_GetRef(MakeShared<FContentReleaseRequest, ESPMode::ThreadSafe>());
+			NewRequest->BundleName = Context.BundleName;
+			NewRequest->LogVerbosityOverride = Context.LogVerbosityOverride;
+			NewRequest->CompleteCallback = MoveTemp(Context.CompleteCallback);
+		}
+		else
+		{
+			LOG_SOURCE_CHUNKINSTALL_OVERRIDE(Context.LogVerbosityOverride, Display, TEXT("Bundle remove %s failed"), *Context.BundleName.ToString());
+			bFailed = true;
+		}
+	}
+
+	if (bFailed)
+	{
+		// send the failure callback immediately
+		FInstallBundleSourceReleaseContentResultInfo ResultInfo;
+		ResultInfo.BundleName = Context.BundleName;
+		ResultInfo.Result = EInstallBundleReleaseResult::OK;
+		Context.CompleteCallback.ExecuteIfBound(AsShared(), MoveTemp(ResultInfo));
+	}
+}
+
+
+void FInstallBundleSourcePlatformChunkInstall::CancelBundles(TArrayView<const FName> BundleNames)
+{
+	CSV_SCOPED_TIMING_STAT(InstallBundleManager, FInstallBundleSourcePlatformChunkInstall_CancelBundles);
+
+	for ( FName BundleName : BundleNames)
+	{
+		FContentRequestRef* ContentRequestPtr = ContentRequests.FindByPredicate( [BundleName](const FContentRequestRef& ContentRequest) {  return ContentRequest->BundleName == BundleName; } );
+		if (ContentRequestPtr != nullptr)
+		{
+			FContentRequestRef ContentRequest = (*ContentRequestPtr);
+			ContentRequest->bCancelled = true;
+		}
+	}
+}
+
+
+TOptional<FInstallBundleSourceProgress> FInstallBundleSourcePlatformChunkInstall::GetBundleProgress(FName BundleName) const
+{
+	TOptional<FInstallBundleSourceProgress> Status;
+
+	const FBundleInfo* BundleInfo = BundleInfoMap.Find(BundleName);
+	if (BundleInfo != nullptr)
+	{
+		float ProgressPercent = PlatformChunkInstall->GetNamedChunkProgress(BundleInfo->NamedChunk, EChunkProgressReportingType::PercentageComplete);
+
+		Status.Emplace();
+		Status->Install_Percent = (ProgressPercent / 100.0f);
+	}
+
+	return Status;
+}
+
+void FInstallBundleSourcePlatformChunkInstall::OnNamedChunkInstall(FName NamedChunk, bool bInstalled)
+{
+	FName BundleName = GetBundleNameForNamedChunk(NamedChunk);
+
+	// see if this named chunk was being installed
+	FContentRequestRef* ContentRequestPtr = ContentRequests.FindByPredicate( [BundleName](const FContentRequestRef& ContentRequest) {  return ContentRequest->BundleName == BundleName; } );
+	if (ContentRequestPtr != nullptr)
+	{
+		FContentRequestRef Request = (*ContentRequestPtr);
+
+		if (Request->bCancelled)
+		{
+			LOG_SOURCE_CHUNKINSTALL_OVERRIDE(Request->LogVerbosityOverride, Display, TEXT("Bundle %s cancelled"), *Request->BundleName.ToString());
+
+			FInstallBundleSourceUpdateContentResultInfo ResultInfo;
+			ResultInfo.BundleName = Request->BundleName;
+			ResultInfo.Result = EInstallBundleResult::UserCancelledError;
+
+			Request->CompleteCallback.ExecuteIfBound(AsShared(), MoveTemp(ResultInfo));
+		}
+		else
+		{
+			if (bInstalled)
+			{
+				LOG_SOURCE_CHUNKINSTALL_OVERRIDE(Request->LogVerbosityOverride, Display, TEXT("Bundle %s finished successfull"), *Request->BundleName.ToString());
+			}
+			else
+			{
+				LOG_SOURCE_CHUNKINSTALL_OVERRIDE(Request->LogVerbosityOverride, Display, TEXT("Bundle %s finished but nothing was installed"), *Request->BundleName.ToString());
+			}
+
+			// send the completion callback
+			FInstallBundleSourceUpdateContentResultInfo ResultInfo;
+			ResultInfo.BundleName = Request->BundleName;
+			ResultInfo.Result = EInstallBundleResult::OK;
+			if (bInstalled && Request->ContentPaths.Num() > 0)
+			{
+				ResultInfo.ContentPaths = Request->ContentPaths;
+				ResultInfo.bContentWasInstalled = true;
+			}
+
+			Request->CompleteCallback.ExecuteIfBound(AsShared(), MoveTemp(ResultInfo));
+		}
+
+		// mark the request for removal next tick
+		Request->bInProgress = false;
+	}
+
+
+	FContentReleaseRequestRef* ContentReleaseRequestPtr = ContentReleaseRequests.FindByPredicate( [BundleName](const FContentReleaseRequestRef& ContentRequest) {  return ContentRequest->BundleName == BundleName; } );
+	if (ContentReleaseRequestPtr != nullptr)
+	{
+		FContentReleaseRequestRef Request = (*ContentReleaseRequestPtr);
+		LOG_SOURCE_CHUNKINSTALL_OVERRIDE(Request->LogVerbosityOverride, Display, TEXT("Bundle %s removed"), *Request->BundleName.ToString());
+
+		// send the completion callback
+		FInstallBundleSourceReleaseContentResultInfo ResultInfo;
+		ResultInfo.BundleName = Request->BundleName;
+		ResultInfo.Result = EInstallBundleReleaseResult::OK;
+		ResultInfo.bContentWasRemoved = !bInstalled;
+		Request->CompleteCallback.ExecuteIfBound(AsShared(), MoveTemp(ResultInfo));
+
+		// mark the request for removal next tick
+		Request->bInProgress = false;
+	}
+}
+
+
+FName FInstallBundleSourcePlatformChunkInstall::GetBundleNameForNamedChunk(FName NamedChunk) const
+{
+	for (const TPair<FName,FBundleInfo>& Pairs : BundleInfoMap)
+	{
+		if (Pairs.Value.NamedChunk == NamedChunk)
+		{
+			return Pairs.Key;
+		}
+	}
+
+	return NAME_None;
+}
+
+#endif //WITH_PLATFORM_INSTALL_BUNDLE_SOURCE
diff --git a/Engine/Plugins/Experimental/DefaultInstallBundleManager/Source/Public/InstallBundleSourcePlatformChunkInstall.h b/Engine/Plugins/Experimental/DefaultInstallBundleManager/Source/Public/InstallBundleSourcePlatformChunkInstall.h
new file mode 100644
--- /dev/null
+++ b/Engine/Plugins/Experimental/DefaultInstallBundleManager/Source/Public/InstallBundleSourcePlatformChunkInstall.h
@@ -0,0 +1,114 @@
+// Copyright Epic Games, Inc. All Rights Reserved.
+
+#pragma once
+
+#include "InstallBundleSourcePlatformBase.h"
+
+#if WITH_PLATFORM_INSTALL_BUNDLE_SOURCE
+#include "Containers/StaticArray.h"
+#include "Containers/Ticker.h"
+#include "GenericPlatform/GenericPlatformChunkInstall.h"
+
+class DEFAULTINSTALLBUNDLEMANAGER_API FInstallBundleSourcePlatformChunkInstall : public FInstallBundleSourcePlatformBase
+{
+private:
+
+	// Internal Types
+	struct FBundleInfo
+	{
+	public:
+		EInstallBundlePriority Priority = EInstallBundlePriority::Low;
+
+		FName NamedChunk;
+		TArray<FString> FilePaths;
+	};
+	
+
+	class FContentRequest
+	{
+	public:
+		FName BundleName;
+		ELogVerbosity::Type LogVerbosityOverride = ELogVerbosity::NoLogging;
+		bool bInProgress = true;
+		bool bCancelled = false;
+
+		TArray<FString> ContentPaths;
+		FInstallBundleCompleteDelegate CompleteCallback;
+	};
+	using FContentRequestRef = TSharedRef<FContentRequest, ESPMode::ThreadSafe>;
+	using FContentRequestPtr = TSharedPtr<FContentRequest, ESPMode::ThreadSafe>;
+	using FContentRequestWeakPtr = TWeakPtr<FContentRequest, ESPMode::ThreadSafe>;
+
+	class FContentReleaseRequest
+	{
+	public:
+		FName BundleName;
+		ELogVerbosity::Type LogVerbosityOverride = ELogVerbosity::NoLogging;
+		bool bInProgress = true;
+		bool bFailed = false;
+		
+		FInstallBundleRemovedDelegate CompleteCallback;
+	};
+	using FContentReleaseRequestRef = TSharedRef<FContentReleaseRequest, ESPMode::ThreadSafe>;
+	using FContentReleaseRequestPtr = TSharedPtr<FContentReleaseRequest, ESPMode::ThreadSafe>;
+	using FContentReleaseRequestWeakPtr = TWeakPtr<FContentReleaseRequest, ESPMode::ThreadSafe>;
+
+	
+public:
+	FInstallBundleSourcePlatformChunkInstall( IPlatformChunkInstall* InPlatformChunkInstall);
+	FInstallBundleSourcePlatformChunkInstall(const FInstallBundleSourcePlatformChunkInstall& Other) = delete;
+	FInstallBundleSourcePlatformChunkInstall& operator=(const FInstallBundleSourcePlatformChunkInstall& Other) = delete;
+	virtual ~FInstallBundleSourcePlatformChunkInstall();
+
+private:
+
+	bool Tick(float dt);
+
+	void TickUpdateChunkOrder();
+
+	// IInstallBundleSource Interface
+public:
+	virtual FInstallBundleSourceInitInfo Init(
+		TSharedRef<InstallBundleUtil::FContentRequestStatsMap> InRequestStats,
+		TSharedPtr<IAnalyticsProviderET> AnalyticsProvider,
+		TSharedPtr<InstallBundleUtil::PersistentStats::FPersistentStatContainerBase> InPersistentStatsContainer) override;
+
+	virtual void AsyncInit(FInstallBundleSourceInitDelegate Callback) override;
+
+	virtual void GetContentState(TArrayView<const FName> BundleNames, EInstallBundleGetContentStateFlags Flags, FInstallBundleGetContentStateDelegate Callback) override;
+
+	virtual EInstallBundleSourceBundleSkipReason GetBundleSkipReason(FName BundleName) const override;
+
+	virtual void RequestUpdateContent(FRequestUpdateContentBundleContext Context) override;
+
+	virtual void RequestReleaseContent(FRequestReleaseContentBundleContext BundleContext) override;
+
+	virtual void CancelBundles(TArrayView<const FName> BundleNames) override;
+
+	virtual TOptional<FInstallBundleSourceProgress> GetBundleProgress(FName BundleName) const override;
+
+protected:
+	virtual bool QueryPersistentBundleInfo(FInstallBundleSourcePersistentBundleInfo& SourceBundleInfo) const override;
+
+	void OnNamedChunkInstall(FName NamedChunk, bool bInstalled);
+
+	FName GetBundleNameForNamedChunk(FName NamedChunk) const;
+
+private:
+	FTSTicker::FDelegateHandle TickHandle;
+	FDelegateHandle NamedChunkInstallDelegateHandle;
+
+	TMap<FName, FBundleInfo> BundleInfoMap;
+	TSet<FName> NamedChunks;
+
+	TArray<FContentRequestRef> ContentRequests;
+	TArray<FContentReleaseRequestRef> ContentReleaseRequests;
+
+	TArray<TUniquePtr<InstallBundleUtil::FInstallBundleTask>> GeneralAsyncTasks;
+
+	IPlatformChunkInstall* PlatformChunkInstall;
+
+	bool bChunkOrderDirty = false;
+	bool bIsUpdatingChunkOrder = false;
+};
+#endif //WITH_PLATFORM_INSTALL_BUNDLE_SOURCE
diff --git a/Engine/Source/Runtime/Core/Public/GenericPlatform/GenericPlatformChunkInstall.h b/Engine/Source/Runtime/Core/Public/GenericPlatform/GenericPlatformChunkInstall.h
--- a/Engine/Source/Runtime/Core/Public/GenericPlatform/GenericPlatformChunkInstall.h
+++ b/Engine/Source/Runtime/Core/Public/GenericPlatform/GenericPlatformChunkInstall.h
@@ -129,6 +129,13 @@ enum class ENamedChunkType : uint8
 	Language,
 };
 
+struct FChunkInstallationStatusDetail
+{
+	uint64 CurrentInstallSize;
+	uint64 FullInstallSize;
+	bool bIsInstalled;
+};
+
 /**
 * Interface for platform specific chunk based install
 **/
@@ -306,6 +313,44 @@ public:
 	 */
 	virtual void RemoveNamedChunkInstallDelegate( FDelegateHandle Delegate ) = 0;
 
+
+
+	/** 
+	 * Returns whether this platform chunk installer implements all the API functions to support the platform chunk install bundle source
+	 */
+	virtual bool SupportsBundleSource() const = 0;
+
+	/** 
+	 * Set whether pak files are auto-mounted when they are installed (the default is that they are mounted)
+	 * @param bEnabled	Whether to auto-mount pak files
+	 * @return			false if this function is not supported
+	 */
+	virtual bool SetAutoPakMountingEnabled( bool bEnabled ) = 0;
+	
+	/**
+	 * Get the list of pak files in the given named chunk.
+	 * @param NamedChunk		The named chunk to query
+	 * @param OutFilesInChunk	The pak files names in the named chunk
+	 * @return					true if the named chunk is valid and this function is supported
+	 */
+	virtual bool GetPakFilesInNamedChunk( const FName NamedChunk, TArray<FString>& OutFilesInChunk) const = 0;
+
+	/**
+	 * Get detailed installation status for the given named chunk
+	 * @param NamedChunk			The named chunk to query
+	 * @param OutChunkStatusDetail	(out) structure that will contain the status detail
+	 * @return						true if the structure has been filled in
+	 */
+	virtual bool GetNamedChunkInstallationStatus( const FName NamedChunk, FChunkInstallationStatusDetail& OutChunkStatusDetail ) const = 0;
+
+	/**
+	 * Determine if the given named chunk is suitable for the current system locale
+	 * @param NamedChunk			The named chunk to query
+	 * @returns						false if the chunk is associated with a different locale
+	 */
+	virtual bool IsNamedChunkForCurrentLocale( const FName NamedChunk ) const = 0;
+	
+
 protected:
 		/**
 		 * Get the current location of a chunk.
@@ -466,6 +511,30 @@ public:
 		NamedChunkInstallDelegate.Remove(Delegate);
 	}
 
+	virtual bool SupportsBundleSource() const override 
+	{ 
+		return false; 
+	}
+
+	virtual bool SetAutoPakMountingEnabled( bool bEnabled ) 
+	{ 
+		return false; 
+	}
+
+	virtual bool GetPakFilesInNamedChunk( const FName NamedChunk, TArray<FString>& OutFilesInChunk) const override 
+	{ 
+		return false; 
+	}
+
+	virtual bool GetNamedChunkInstallationStatus( const FName NamedChunk, FChunkInstallationStatusDetail& OutChunkStatusDetail ) const override 
+	{ 
+		return false; 
+	}
+
+	virtual bool IsNamedChunkForCurrentLocale( const FName NamedChunk ) const 
+	{ 
+		return true; 
+	}
 
 protected:
 
diff --git a/Engine/Source/Runtime/PakFile/Private/IPlatformFilePak.cpp b/Engine/Source/Runtime/PakFile/Private/IPlatformFilePak.cpp
--- a/Engine/Source/Runtime/PakFile/Private/IPlatformFilePak.cpp
+++ b/Engine/Source/Runtime/PakFile/Private/IPlatformFilePak.cpp
@@ -72,6 +72,9 @@ int64 GTotalLoadedLastTick = 0;
 #ifndef HAS_PLATFORM_PAK_INSTALL_CHECK
 #define HAS_PLATFORM_PAK_INSTALL_CHECK 0
 #endif
+#ifndef ENABLE_PLATFORM_CHUNK_INSTALL
+#define ENABLE_PLATFORM_CHUNK_INSTALL 1
+#endif
 #ifndef ALL_PAKS_WILDCARD
 #define ALL_PAKS_WILDCARD "*.pak"
 #endif 
@@ -5260,6 +5263,7 @@ bool FPakPlatformFile::IsNonPakFilenameAllowed(const FString& InFilename)
 #if !HAS_PLATFORM_PAK_INSTALL_CHECK
 bool FPakPlatformFile::IsPakFileInstalled(const FString& InFilename)
 {
+#if ENABLE_PLATFORM_CHUNK_INSTALL
 	IPlatformChunkInstall* ChunkInstall = FPlatformMisc::GetPlatformChunkInstall();
 	if (ChunkInstall)
 	{
@@ -5273,6 +5277,7 @@ bool FPakPlatformFile::IsPakFileInstalled(const FString& InFilename)
 			}
 		}
 	}
+#endif
 
 	return true;
 }
