Platform chunk install bundle wrapper can use PlatformChunkName property to override the named chunk that a bundle will use.

 - useful on platforms that do not support a lot of named chunks, allowing multple bundles to reference the same named chunk.
 - also allows XYZOptional bundles to work again, linking them to their main XYZ bundle state on platforms they're not supported on.

#rnx

diff --git a/Engine/Plugins/Experimental/DefaultInstallBundleManager/Source/Private/InstallBundleSourcePlatformChunkInstall.cpp b/Engine/Plugins/Experimental/DefaultInstallBundleManager/Source/Private/InstallBundleSourcePlatformChunkInstall.cpp
--- a/Engine/Plugins/Experimental/DefaultInstallBundleManager/Source/Private/InstallBundleSourcePlatformChunkInstall.cpp
+++ b/Engine/Plugins/Experimental/DefaultInstallBundleManager/Source/Private/InstallBundleSourcePlatformChunkInstall.cpp
@@ -199,16 +199,20 @@ void FInstallBundleSourcePlatformChunkInstall::AsyncInit(FInstallBundleSourceIni
 				continue;
 			}
 
+			FString PlatformChunkName;
+			InstallBundleConfig->GetString(*Section, TEXT("PlatformChunkName"), PlatformChunkName);
+			
 			int32 PlatformChunkID = 0;
-			if (InstallBundleConfig->GetInt(*Section, TEXT("PlatformChunkID"), PlatformChunkID) && PlatformChunkID < 0)
+			if (PlatformChunkName.IsEmpty() && InstallBundleConfig->GetInt(*Section, TEXT("PlatformChunkID"), PlatformChunkID) && PlatformChunkID < 0)
 			{
 				continue;
 			}
+			//... NB PlatformChunkID is ignored by this bundle source, and only used to mark the bundle as a platform bundle
 
 			// create bundle info
 			FName BundleName( *Section.RightChop(InstallBundleUtil::GetInstallBundleSectionPrefix().Len()));
 			FBundleInfo& BundleInfo = Context->BundleInfoMap.Add(BundleName);
-			BundleInfo.NamedChunk = BundleName;
+			BundleInfo.NamedChunk = PlatformChunkName.IsEmpty() ? BundleName : FName(*PlatformChunkName);
 			BundleInfo.Priority = EInstallBundlePriority::Normal;
 
 			FString PriorityString;
@@ -585,13 +589,13 @@ TOptional<FInstallBundleSourceProgress> FInstallBundleSourcePlatformChunkInstall
 
 void FInstallBundleSourcePlatformChunkInstall::OnNamedChunkInstall(FName NamedChunk, bool bInstalled)
 {
-	FName BundleName = GetBundleNameForNamedChunk(NamedChunk);
-
-	// see if this named chunk was being installed
-	FContentRequestRef* ContentRequestPtr = ContentRequests.FindByPredicate( [BundleName](const FContentRequestRef& ContentRequest) {  return ContentRequest->BundleName == BundleName; } );
-	if (ContentRequestPtr != nullptr)
+	// see if this named chunk was being installed by any active requests
+	for ( FContentRequestRef Request : ContentRequests)
 	{
-		FContentRequestRef Request = (*ContentRequestPtr);
+		if (GetNamedChunkForBundle(Request->BundleName) != NamedChunk)
+		{
+			continue;
+		}
 
 		if (Request->bCancelled)
 		{
@@ -631,11 +635,14 @@ void FInstallBundleSourcePlatformChunkInstall::OnNamedChunkInstall(FName NamedCh
 		Request->bInProgress = false;
 	}
 
-
-	FContentReleaseRequestRef* ContentReleaseRequestPtr = ContentReleaseRequests.FindByPredicate( [BundleName](const FContentReleaseRequestRef& ContentRequest) {  return ContentRequest->BundleName == BundleName; } );
-	if (ContentReleaseRequestPtr != nullptr)
+	// see if this named chunk was being released by any active requests
+	for (FContentReleaseRequestRef Request : ContentReleaseRequests)
 	{
-		FContentReleaseRequestRef Request = (*ContentReleaseRequestPtr);
+		if (GetNamedChunkForBundle(Request->BundleName) != NamedChunk)
+		{
+			continue;
+		}
+
 		LOG_SOURCE_CHUNKINSTALL_OVERRIDE(Request->LogVerbosityOverride, Display, TEXT("Bundle %s removed"), *Request->BundleName.ToString());
 
 		// send the completion callback
@@ -651,14 +658,12 @@ void FInstallBundleSourcePlatformChunkInstall::OnNamedChunkInstall(FName NamedCh
 }
 
 
-FName FInstallBundleSourcePlatformChunkInstall::GetBundleNameForNamedChunk(FName NamedChunk) const
+FName FInstallBundleSourcePlatformChunkInstall::GetNamedChunkForBundle(FName BundleName) const
 {
-	for (const TPair<FName,FBundleInfo>& Pairs : BundleInfoMap)
+	const FBundleInfo* BundleInfo = BundleInfoMap.Find(BundleName);
+	if (BundleInfo != nullptr)
 	{
-		if (Pairs.Value.NamedChunk == NamedChunk)
-		{
-			return Pairs.Key;
-		}
+		return BundleInfo->NamedChunk;
 	}
 
 	return NAME_None;
diff --git a/Engine/Plugins/Experimental/DefaultInstallBundleManager/Source/Public/InstallBundleSourcePlatformChunkInstall.h b/Engine/Plugins/Experimental/DefaultInstallBundleManager/Source/Public/InstallBundleSourcePlatformChunkInstall.h
--- a/Engine/Plugins/Experimental/DefaultInstallBundleManager/Source/Public/InstallBundleSourcePlatformChunkInstall.h
+++ b/Engine/Plugins/Experimental/DefaultInstallBundleManager/Source/Public/InstallBundleSourcePlatformChunkInstall.h
@@ -92,7 +92,7 @@ protected:
 
 	void OnNamedChunkInstall(FName NamedChunk, bool bInstalled);
 
-	FName GetBundleNameForNamedChunk(FName NamedChunk) const;
+	FName GetNamedChunkForBundle(FName BundleName) const;
 
 private:
 	FTSTicker::FDelegateHandle TickHandle;
diff --git a/Engine/Source/Runtime/InstallBundleManager/Private/InstallBundleUtils.cpp b/Engine/Source/Runtime/InstallBundleManager/Private/InstallBundleUtils.cpp
--- a/Engine/Source/Runtime/InstallBundleManager/Private/InstallBundleUtils.cpp
+++ b/Engine/Source/Runtime/InstallBundleManager/Private/InstallBundleUtils.cpp
@@ -73,8 +73,11 @@ namespace InstallBundleUtil
 
 	bool IsPlatformInstallBundlePredicate(const FConfigFile& InstallBundleConfig, const FString& Section)
 	{
+		FString PlatformChunkName;
+		InstallBundleConfig.GetString(*Section, TEXT("PlatformChunkName"), PlatformChunkName);
+
 		int32 ChunkID = 0;
-		if (InstallBundleConfig.GetInt(*Section, TEXT("PlatformChunkID"), ChunkID) && ChunkID < 0)
+		if (PlatformChunkName.IsEmpty() && InstallBundleConfig.GetInt(*Section, TEXT("PlatformChunkID"), ChunkID) && ChunkID < 0)
 		{
 			return false;
 		}
