Follow up for Eric's fix
- Capture by ref, not value
- Also update bulk source

diff --git a/Engine/Plugins/Experimental/DefaultInstallBundleManager/Source/Private/InstallBundleSourceBulk.cpp b/Engine/Plugins/Experimental/DefaultInstallBundleManager/Source/Private/InstallBundleSourceBulk.cpp
--- a/Engine/Plugins/Experimental/DefaultInstallBundleManager/Source/Private/InstallBundleSourceBulk.cpp
+++ b/Engine/Plugins/Experimental/DefaultInstallBundleManager/Source/Private/InstallBundleSourceBulk.cpp
@@ -322,17 +322,17 @@ void FInstallBundleSourceBulk::AsyncInit_MakeBundlesForBulkBuild()
 			if (!bDidLoadAllFilesFromMetadata)
 			{
 				// Bundle regex need to be applied in order
-				SectionNames.StableSort([this](const FString& SectionA, const FString& SectionB) -> bool
+				SectionNames.StableSort([InstallBundleConfig](const FString& SectionA, const FString& SectionB) -> bool
 				{
-					int BundleAOrder = INT_MAX;
-					int BundleBOrder = INT_MAX;
+					int32 BundleAOrder = INT_MAX;
+					int32 BundleBOrder = INT_MAX;
 
-					if (!GConfig->GetInt(*SectionA, TEXT("Order"), BundleAOrder, GInstallBundleIni))
+					if (!InstallBundleConfig->GetInt(*SectionA, TEXT("Order"), BundleAOrder))
 					{
 						LOG_SOURCE_BULK(Warning, TEXT("Bundle Section %s doesn't have an order"), *SectionA);
 					}
 
-					if (!GConfig->GetInt(*SectionB, TEXT("Order"), BundleBOrder, GInstallBundleIni))
+					if (!InstallBundleConfig->GetInt(*SectionB, TEXT("Order"), BundleBOrder))
 					{
 						LOG_SOURCE_BULK(Warning, TEXT("Bundle Section %s doesn't have an order"), *SectionB);
 					}
diff --git a/Engine/Source/Runtime/InstallBundleManager/Private/InstallBundleUtils.cpp b/Engine/Source/Runtime/InstallBundleManager/Private/InstallBundleUtils.cpp
--- a/Engine/Source/Runtime/InstallBundleManager/Private/InstallBundleUtils.cpp
+++ b/Engine/Source/Runtime/InstallBundleManager/Private/InstallBundleUtils.cpp
@@ -115,7 +115,7 @@ namespace InstallBundleUtil
 			BundleRegexList.Emplace(TPair<FString, TArray<FRegexPattern>>(BundleName, MoveTemp(SearchRegexPatterns)));
 		}
 
-		BundleRegexList.StableSort([InstallBundleConfig](const TPair<FString, TArray<FRegexPattern>>& PairA, const TPair<FString, TArray<FRegexPattern>>& PairB) -> bool
+		BundleRegexList.StableSort([&InstallBundleConfig](const TPair<FString, TArray<FRegexPattern>>& PairA, const TPair<FString, TArray<FRegexPattern>>& PairB) -> bool
 		{
 			int32 BundleAOrder = INT_MAX;
 			int32 BundleBOrder = INT_MAX;
