Add a more informative named chunk completion callback & deprecated the old one

 - new callback can report installation state in addition to success/failure
 - added new helper function to allow platform chunk installers to trigger the new and old callbacks.

diff --git a/Engine/Plugins/Experimental/DefaultInstallBundleManager/Source/Private/InstallBundleSourcePlatformChunkInstall.cpp b/Engine/Plugins/Experimental/DefaultInstallBundleManager/Source/Private/InstallBundleSourcePlatformChunkInstall.cpp
--- a/Engine/Plugins/Experimental/DefaultInstallBundleManager/Source/Private/InstallBundleSourcePlatformChunkInstall.cpp
+++ b/Engine/Plugins/Experimental/DefaultInstallBundleManager/Source/Private/InstallBundleSourcePlatformChunkInstall.cpp
@@ -25,7 +25,7 @@ FInstallBundleSourcePlatformChunkInstall::FInstallBundleSourcePlatformChunkInsta
 {
 	TickHandle = FTSTicker::GetCoreTicker().AddTicker(FTickerDelegate::CreateRaw(this, &FInstallBundleSourcePlatformChunkInstall::Tick));
 
-	NamedChunkInstallDelegateHandle = PlatformChunkInstall->AddNamedChunkInstallDelegate( FPlatformNamedChunkInstallDelegate::CreateRaw(this, &FInstallBundleSourcePlatformChunkInstall::OnNamedChunkInstall));
+	NamedChunkInstallDelegateHandle = PlatformChunkInstall->AddNamedChunkCompleteDelegate( FPlatformNamedChunkCompleteDelegate::CreateRaw(this, &FInstallBundleSourcePlatformChunkInstall::OnNamedChunkInstall));
 
 	InPlatformChunkInstall->SetAutoPakMountingEnabled(false);
 }
@@ -33,7 +33,7 @@ FInstallBundleSourcePlatformChunkInstall::FInstallBundleSourcePlatformChunkInsta
 
 FInstallBundleSourcePlatformChunkInstall::~FInstallBundleSourcePlatformChunkInstall()
 {
-	PlatformChunkInstall->RemoveNamedChunkInstallDelegate(NamedChunkInstallDelegateHandle);
+	PlatformChunkInstall->RemoveNamedChunkCompleteDelegate(NamedChunkInstallDelegateHandle);
 	NamedChunkInstallDelegateHandle.Reset();
 
 	FTSTicker::GetCoreTicker().RemoveTicker(TickHandle);
@@ -587,12 +587,12 @@ TOptional<FInstallBundleSourceProgress> FInstallBundleSourcePlatformChunkInstall
 	return Status;
 }
 
-void FInstallBundleSourcePlatformChunkInstall::OnNamedChunkInstall(FName NamedChunk, bool bInstalled)
+void FInstallBundleSourcePlatformChunkInstall::OnNamedChunkInstall( const FNamedChunkCompleteCallbackParam& Param )
 {
 	// see if this named chunk was being installed by any active requests
 	for ( FContentRequestRef Request : ContentRequests)
 	{
-		if (GetNamedChunkForBundle(Request->BundleName) != NamedChunk)
+		if (GetNamedChunkForBundle(Request->BundleName) != Param.NamedChunk)
 		{
 			continue;
 		}
@@ -609,20 +609,13 @@ void FInstallBundleSourcePlatformChunkInstall::OnNamedChunkInstall(FName NamedCh
 		}
 		else
 		{
-			if (bInstalled)
-			{
-				LOG_SOURCE_CHUNKINSTALL_OVERRIDE(Request->LogVerbosityOverride, Display, TEXT("Bundle %s finished successfull"), *Request->BundleName.ToString());
-			}
-			else
-			{
-				LOG_SOURCE_CHUNKINSTALL_OVERRIDE(Request->LogVerbosityOverride, Display, TEXT("Bundle %s finished but nothing was installed"), *Request->BundleName.ToString());
-			}
+			LOG_SOURCE_CHUNKINSTALL_OVERRIDE(Request->LogVerbosityOverride, Display, TEXT("Bundle request %s finished. IsInstalled: %s, Succeeded: %s"), *Request->BundleName.ToString(), *LexToString(Param.bIsInstalled), *LexToString(Param.bHasSucceeded) );
 
 			// send the completion callback
 			FInstallBundleSourceUpdateContentResultInfo ResultInfo;
 			ResultInfo.BundleName = Request->BundleName;
 			ResultInfo.Result = EInstallBundleResult::OK;
-			if (bInstalled && Request->ContentPaths.Num() > 0)
+			if (Param.bIsInstalled && Param.bHasSucceeded && Request->ContentPaths.Num() > 0)
 			{
 				ResultInfo.ContentPaths = Request->ContentPaths;
 				ResultInfo.bContentWasInstalled = true;
@@ -638,18 +631,18 @@ void FInstallBundleSourcePlatformChunkInstall::OnNamedChunkInstall(FName NamedCh
 	// see if this named chunk was being released by any active requests
 	for (FContentReleaseRequestRef Request : ContentReleaseRequests)
 	{
-		if (GetNamedChunkForBundle(Request->BundleName) != NamedChunk)
+		if (GetNamedChunkForBundle(Request->BundleName) != Param.NamedChunk)
 		{
 			continue;
 		}
 
-		LOG_SOURCE_CHUNKINSTALL_OVERRIDE(Request->LogVerbosityOverride, Display, TEXT("Bundle %s removed"), *Request->BundleName.ToString());
+		LOG_SOURCE_CHUNKINSTALL_OVERRIDE(Request->LogVerbosityOverride, Display, TEXT("Bundle remove request %s finished. IsInstalled: %s, Succeeded: %s"), *Request->BundleName.ToString(), *LexToString(Param.bIsInstalled), *LexToString(Param.bHasSucceeded) );
 
 		// send the completion callback
 		FInstallBundleSourceReleaseContentResultInfo ResultInfo;
 		ResultInfo.BundleName = Request->BundleName;
 		ResultInfo.Result = EInstallBundleReleaseResult::OK;
-		ResultInfo.bContentWasRemoved = !bInstalled;
+		ResultInfo.bContentWasRemoved = !Param.bIsInstalled && Param.bHasSucceeded;
 		Request->CompleteCallback.ExecuteIfBound(AsShared(), MoveTemp(ResultInfo));
 
 		// mark the request for removal next tick
diff --git a/Engine/Plugins/Experimental/DefaultInstallBundleManager/Source/Public/InstallBundleSourcePlatformChunkInstall.h b/Engine/Plugins/Experimental/DefaultInstallBundleManager/Source/Public/InstallBundleSourcePlatformChunkInstall.h
--- a/Engine/Plugins/Experimental/DefaultInstallBundleManager/Source/Public/InstallBundleSourcePlatformChunkInstall.h
+++ b/Engine/Plugins/Experimental/DefaultInstallBundleManager/Source/Public/InstallBundleSourcePlatformChunkInstall.h
@@ -90,7 +90,7 @@ public:
 protected:
 	virtual bool QueryPersistentBundleInfo(FInstallBundleSourcePersistentBundleInfo& SourceBundleInfo) const override;
 
-	void OnNamedChunkInstall(FName NamedChunk, bool bInstalled);
+	void OnNamedChunkInstall(const FNamedChunkCompleteCallbackParam& Param);
 
 	FName GetNamedChunkForBundle(FName BundleName) const;
 
diff --git a/Engine/Source/Runtime/Core/Private/GenericPlatform/GenericPlatformChunkInstall.cpp b/Engine/Source/Runtime/Core/Private/GenericPlatform/GenericPlatformChunkInstall.cpp
--- a/Engine/Source/Runtime/Core/Private/GenericPlatform/GenericPlatformChunkInstall.cpp
+++ b/Engine/Source/Runtime/Core/Private/GenericPlatform/GenericPlatformChunkInstall.cpp
@@ -4,3 +4,27 @@
 
 DEFINE_LOG_CATEGORY(LogChunkInstaller)
 
+void FGenericPlatformChunkInstall::DoNamedChunkCompleteCallbacks( const FName NamedChunk, EChunkLocation::Type Location, bool bHasSucceeded )
+{
+	check(IsInGameThread());
+
+	bool bIsInstalled = (Location == EChunkLocation::LocalFast) || (Location == EChunkLocation::LocalSlow);
+
+	if (NamedChunkCompleteDelegate.IsBound())
+	{
+		FNamedChunkCompleteCallbackParam Param;
+		Param.NamedChunk = NamedChunk;
+		Param.Location = Location;
+		Param.bIsInstalled = bIsInstalled;
+		Param.bHasSucceeded = bHasSucceeded;
+
+		NamedChunkCompleteDelegate.Broadcast(Param);
+	}
+
+PRAGMA_DISABLE_DEPRECATION_WARNINGS
+	if (bIsInstalled && NamedChunkInstallDelegate.IsBound())
+	{
+		NamedChunkInstallDelegate.Broadcast(NamedChunk, bHasSucceeded);
+	}
+PRAGMA_ENABLE_DEPRECATION_WARNINGS
+}
diff --git a/Engine/Source/Runtime/Core/Public/GenericPlatform/GenericPlatformChunkInstall.h b/Engine/Source/Runtime/Core/Public/GenericPlatform/GenericPlatformChunkInstall.h
--- a/Engine/Source/Runtime/Core/Public/GenericPlatform/GenericPlatformChunkInstall.h
+++ b/Engine/Source/Runtime/Core/Public/GenericPlatform/GenericPlatformChunkInstall.h
@@ -65,6 +65,14 @@ namespace EChunkProgressReportingType
 	};
 }
 
+struct FNamedChunkCompleteCallbackParam
+{
+	FName NamedChunk;
+	EChunkLocation::Type Location;
+	bool bIsInstalled;
+	bool bHasSucceeded;
+};
+
 /**
  * Platform Chunk Install Module Interface
  */
@@ -82,10 +90,13 @@ DECLARE_DELEGATE_OneParam(FPlatformChunkInstallCompleteDelegate, uint32);
 DECLARE_DELEGATE_TwoParams(FPlatformChunkInstallDelegate, uint32, bool);
 DECLARE_MULTICAST_DELEGATE_TwoParams(FPlatformChunkInstallMultiDelegate, uint32, bool);
 
-/** Delegate called when a Named Chunk either successfully installs or fails to install, bool is success */
+/** Deprecated delegate called when a Named Chunk either successfully installs or fails to install, bool is success */
 DECLARE_DELEGATE_TwoParams(FPlatformNamedChunkInstallDelegate, FName, bool);
 DECLARE_MULTICAST_DELEGATE_TwoParams(FPlatformNamedChunkInstallMultiDelegate, FName, bool);
 
+/** Delegate called when a Named Chunk either successfully installs or fails to install */
+DECLARE_DELEGATE_OneParam(FPlatformNamedChunkCompleteDelegate, const FNamedChunkCompleteCallbackParam&);
+DECLARE_MULTICAST_DELEGATE_OneParam(FPlatformNamedChunkCompleteMultiDelegate, const FNamedChunkCompleteCallbackParam&);
 
 enum class ECustomChunkType : uint8
 {
@@ -305,13 +316,28 @@ public:
 	 * @param Delegate		The delegate to call when any named chunk is installed or fails to install
 	 * @return				Handle to the bound delegate
 	 */
+	UE_DEPRECATED(5.4, "use AddNamedChunkCompleteDelegate instead")
 	virtual FDelegateHandle AddNamedChunkInstallDelegate( FPlatformNamedChunkInstallDelegate Delegate ) = 0;
 
 	/**
 	 * Remove a delegate callback on named chunk install completion.
 	 * @param Delegate		The delegate to remove.
 	 */
+	UE_DEPRECATED(5.4, "use RemoveNamedChunkCompleteDelegate instead")
 	virtual void RemoveNamedChunkInstallDelegate( FDelegateHandle Delegate ) = 0;
+	
+	/** 
+	 * Request a delegate callback on named chunk install completion or failure. Request may not be respected.
+	 * @param Delegate		The delegate to call when any named chunk is installed or fails to install
+	 * @return				Handle to the bound delegate
+	 */
+	virtual FDelegateHandle AddNamedChunkCompleteDelegate( FPlatformNamedChunkCompleteDelegate Delegate ) = 0;
+
+	/**
+	 * Remove a delegate callback on named chunk install completion.
+	 * @param Delegate		The delegate to remove.
+	 */
+	virtual void RemoveNamedChunkCompleteDelegate( FDelegateHandle Delegate ) = 0;
 
 
 
@@ -511,6 +537,17 @@ public:
 		NamedChunkInstallDelegate.Remove(Delegate);
 	}
 
+	virtual void RemoveNamedChunkCompleteDelegate(FDelegateHandle Delegate) override
+	{
+		NamedChunkCompleteDelegate.Remove(Delegate);
+	}
+
+	virtual FDelegateHandle AddNamedChunkCompleteDelegate(FPlatformNamedChunkCompleteDelegate Delegate) override
+	{
+		return NamedChunkCompleteDelegate.Add(Delegate);
+	}
+
+
 	virtual bool SupportsBundleSource() const override 
 	{ 
 		return false; 
@@ -538,9 +575,12 @@ public:
 
 protected:
 
+	void DoNamedChunkCompleteCallbacks( const FName NamedChunk, EChunkLocation::Type Location, bool bHasSucceeded );
+
 	/** Delegates called when installation succeeds or fails */
 	FPlatformChunkInstallMultiDelegate InstallDelegate;
 	FPlatformNamedChunkInstallMultiDelegate NamedChunkInstallDelegate;
+	FPlatformNamedChunkCompleteMultiDelegate NamedChunkCompleteDelegate;
 
 	virtual EChunkLocation::Type GetChunkLocation(uint32 ChunkID) override
 	{
