IInstallBundleManager method to change a bundle cache size which overrides ini settings but not the command line arg (for now it doesn't support changing it after the bundles have been initialized)

diff --git a/Engine/Plugins/Experimental/DefaultInstallBundleManager/Source/Private/DefaultInstallBundleManager.cpp b/Engine/Plugins/Experimental/DefaultInstallBundleManager/Source/Private/DefaultInstallBundleManager.cpp
--- a/Engine/Plugins/Experimental/DefaultInstallBundleManager/Source/Private/DefaultInstallBundleManager.cpp
+++ b/Engine/Plugins/Experimental/DefaultInstallBundleManager/Source/Private/DefaultInstallBundleManager.cpp
@@ -3932,6 +3932,14 @@ void FDefaultInstallBundleManager::UpdateContentRequestFlags(TArrayView<const FN
 	}
 }
 
+void FDefaultInstallBundleManager::SetCacheSize(FName CacheName, uint64 CacheSize)
+{
+	if (ensureMsgf(!BundleCaches.Contains(CacheName), TEXT("FDefaultInstallBundleManager::SetCacheSize is only supported prior to initialization, for now")))
+	{
+		BundleCacheSizeOverrides.Add(CacheName, CacheSize);
+	}
+}
+
 void FDefaultInstallBundleManager::StartPatchCheck()
 {
 	if (bIsCheckingForPatch)
@@ -4349,6 +4357,16 @@ void FDefaultInstallBundleManager::AsyncInit_InitBundleCaches()
 		InitInfo.Size = CacheSize;
 	}
 
+	// Apply cache size runtime overrides
+	for (const TPair<FName, uint64>& BundleCacheSizeOverride : BundleCacheSizeOverrides)
+	{
+		FInstallBundleCacheInitInfo* InitInfo = BundleCacheInitInfo.Find(BundleCacheSizeOverride.Key);
+		if (ensureMsgf(InitInfo, TEXT("Size override cannot be applied on cache '%s' because it doesn't exist"), *BundleCacheSizeOverride.Key.ToString()))
+		{
+			InitInfo->Size = BundleCacheSizeOverride.Value;
+		}
+	}
+
 	// Check to override cache size from command line
 	{
 		FString Mapping;
diff --git a/Engine/Plugins/Experimental/DefaultInstallBundleManager/Source/Public/DefaultInstallBundleManager.h b/Engine/Plugins/Experimental/DefaultInstallBundleManager/Source/Public/DefaultInstallBundleManager.h
--- a/Engine/Plugins/Experimental/DefaultInstallBundleManager/Source/Public/DefaultInstallBundleManager.h
+++ b/Engine/Plugins/Experimental/DefaultInstallBundleManager/Source/Public/DefaultInstallBundleManager.h
@@ -653,6 +653,8 @@ public:
 	virtual EInstallBundleRequestFlags GetModifyableContentRequestFlags() const override;
 	virtual void UpdateContentRequestFlags(TArrayView<const FName> BundleNames, EInstallBundleRequestFlags AddFlags, EInstallBundleRequestFlags RemoveFlags) override;
 
+	virtual void SetCacheSize(FName CacheName, uint64 CacheSize) override;
+
 	virtual void StartPatchCheck() override;
 	virtual void AddEnvironmentWantsPatchCheckBackCompatDelegate(FName Tag, FInstallBundleManagerEnvironmentWantsPatchCheck Delegate) override;
 	virtual void RemoveEnvironmentWantsPatchCheckBackCompatDelegate(FName Tag) override;
@@ -712,6 +714,7 @@ protected:
 
 	TMap<FName, TSharedRef<FInstallBundleCache>> BundleCaches;
 	TMap<EInstallBundleSourceType, FName> BundleSourceCaches;
+	TMap<FName, uint64> BundleCacheSizeOverrides;
 
 	TMap<TTuple<EInstallBundleSourceType, FName>, TArray<FCacheEvictionRequestorRef>> PendingCacheEvictions; // (Source, Bundle) -> List of requestors
 	TMap<TTuple<FName, FName>, TArray<EInstallBundleSourceType>> CachesPendingEvictToSources; // (Cache, Bundle) -> List of Sources
diff --git a/Engine/Source/Runtime/InstallBundleManager/Public/InstallBundleManagerInterface.h b/Engine/Source/Runtime/InstallBundleManager/Public/InstallBundleManagerInterface.h
--- a/Engine/Source/Runtime/InstallBundleManager/Public/InstallBundleManagerInterface.h
+++ b/Engine/Source/Runtime/InstallBundleManager/Public/InstallBundleManagerInterface.h
@@ -168,6 +168,8 @@ public:
 	INSTALLBUNDLEMANAGER_API void UpdateContentRequestFlags(FName BundleName, EInstallBundleRequestFlags AddFlags, EInstallBundleRequestFlags RemoveFlags);
 	virtual void UpdateContentRequestFlags(TArrayView<const FName> BundleNames, EInstallBundleRequestFlags AddFlags, EInstallBundleRequestFlags RemoveFlags) = 0;
 	
+	virtual void SetCacheSize(FName CacheName, uint64 CacheSize) = 0;
+
 	INSTALLBUNDLEMANAGER_API virtual void StartPatchCheck();
 	virtual void AddEnvironmentWantsPatchCheckBackCompatDelegate(FName Tag, FInstallBundleManagerEnvironmentWantsPatchCheck Delegate) {}
 	virtual void RemoveEnvironmentWantsPatchCheckBackCompatDelegate(FName Tag) {}
diff --git a/Engine/Source/Runtime/NullInstallBundleManager/Private/NullInstallBundleManager.cpp b/Engine/Source/Runtime/NullInstallBundleManager/Private/NullInstallBundleManager.cpp
--- a/Engine/Source/Runtime/NullInstallBundleManager/Private/NullInstallBundleManager.cpp
+++ b/Engine/Source/Runtime/NullInstallBundleManager/Private/NullInstallBundleManager.cpp
@@ -135,6 +135,9 @@ class FNullInstallBundleManager : public IInstallBundleManager
 
 	}
 
+	virtual void SetCacheSize(FName CacheName, uint64 CacheSize) override
+	{}
+
 	virtual bool SupportsEarlyStartupPatching() const override
 	{
 		return false;
