GFP state machine can now detect when its possible to force monolithic shader library loading which is faster if we know we don't need to search for chunked shader libraries.
- Can be disabled with CVarAllowForceMonolithicShaderLibrary.
- FGameFeaturePreMountingContext can disable shader library loading in cases where we may want to control it externally.
- Install bundle manager now tracks if a bundle contains chunks
- GetPakchunkIndexFromPakFile now uses string views and string builders instead of temporary strings
- GetPakchunkIndexFromPakFile is no longer duplicated in IPlatformFilePak
- FPakPlatformFile::HandleMountPakDelegate no longer does a linear search to find the just mounted pak file

diff --git a/Engine/Plugins/Experimental/DefaultInstallBundleManager/Source/Private/DefaultInstallBundleManager.cpp b/Engine/Plugins/Experimental/DefaultInstallBundleManager/Source/Private/DefaultInstallBundleManager.cpp
--- a/Engine/Plugins/Experimental/DefaultInstallBundleManager/Source/Private/DefaultInstallBundleManager.cpp
+++ b/Engine/Plugins/Experimental/DefaultInstallBundleManager/Source/Private/DefaultInstallBundleManager.cpp
@@ -1243,6 +1243,21 @@ void FDefaultInstallBundleManager::UpdateBundleSourceComplete(TSharedRef<IInstal
 	}
 	Request->SourceRequestResults.Empty();
 
+	// If there are no content paths, its likely this is a chunk that doesn't exist on the current platfrom, so set bContainsChunks true.
+	// This is a corner case but as far as I know there is no other situation that would allow for an empty bundle that was not chunked.
+	// If such a case were to arise, then bContainsChunks would need to be determined by each bundle source individually.
+	BundleInfo.ContentPaths.bContainsChunks = BundleInfo.ContentPaths.ContentPaths.IsEmpty();
+	for (const FString& ContentPath : BundleInfo.ContentPaths.ContentPaths)
+	{
+		BundleInfo.ContentPaths.bContainsChunks = 
+			ContentPath.EndsWith(TEXTVIEW(".pak")) &&
+			FPlatformMisc::GetPakchunkIndexFromPakFile(ContentPath) != INDEX_NONE;
+		if (BundleInfo.ContentPaths.bContainsChunks)
+		{
+			break;
+		}
+	}
+
 	if (StateSignifiesNeedsInstall(GetBundleStatus(BundleInfo)))
 	{
 		if (Request->Result == EInstallBundleResult::OK)
@@ -1435,7 +1450,7 @@ void FDefaultInstallBundleManager::MountPaks(FContentRequestRef Request)
 
 	LOG_INSTALL_BUNDLE_MAN_OVERRIDE(Request->LogVerbosityOverride, Display, TEXT("Mounting Paks for Request %s"), *BundleInfo.BundleNameString);
 
-	TSharedRef<bool, ESPMode::ThreadSafe> bMountedPaks = MakeShared<bool, ESPMode::ThreadSafe>(false);
+	TSharedRef<bool> bMountedPaks = MakeShared<bool>(false);
 
 	TUniqueFunction<void()> WorkFunc =
 		[ContentPaths = BundleInfo.ContentPaths.ContentPaths
@@ -1484,34 +1499,37 @@ void FDefaultInstallBundleManager::MountPaks(FContentRequestRef Request)
 
 bool FDefaultInstallBundleManager::MountPaksInList(TArrayView<FString> Paths, ELogVerbosity::Type LogVerbosityOverride)
 {
-	if (!FCoreDelegates::MountPak.IsBound())
+	FPakPlatformFile* PakPlatformFile = static_cast<FPakPlatformFile*>(FPlatformFileManager::Get().FindPlatformFile(TEXT("PakFile")));
+	if (!PakPlatformFile)
 	{
 		ensureMsgf(false, TEXT("Pak files have not been correctly initalized. Use -UsePaks on the cmdline if you are using the UnrealEditor.exe"));
 		return false; // if FCoreDelegates::MountPak is unbound there is a major issue.
 	}
 
+	// Sort in descending order.
+	Paths.Sort(TGreater<FString>());
+
 	// Find already mounted paks
-	TArray<FString> MountedPaks;
-	FPakPlatformFile* PakPlatformFile = static_cast<FPakPlatformFile*>(FPlatformFileManager::Get().FindPlatformFile(TEXT("PakFile")));
-	check(PakPlatformFile);
+	TSet<FString> MountedPaks;
 	PakPlatformFile->GetMountedPakFilenames(MountedPaks);
 
-	// Sort in descending order.
 	bool bMountedPaks = false;
-	Paths.Sort(TGreater<FString>());
 	for (const FString& File : Paths)
 	{
-		if (!File.EndsWith(TEXT(".pak")))
+		if (!File.EndsWith(TEXTVIEW(".pak")))
+		{
 			continue;
+		}
 
 		if (MountedPaks.Contains(File))
 		{
-			LOG_INSTALL_BUNDLE_MAN_OVERRIDE(LogVerbosityOverride, Display, TEXT("Pak file: %s already mounted, skipping. \n"), *File);
+			LOG_INSTALL_BUNDLE_MAN_OVERRIDE(LogVerbosityOverride, Warning, TEXT("Pak file: %s already mounted, skipping. \n"), *File);
 			continue;
 		}
 
 		LOG_INSTALL_BUNDLE_MAN_OVERRIDE(LogVerbosityOverride, Verbose, TEXT("Mounting pak file: %s \n"), *File);
-		if (FCoreDelegates::MountPak.Execute(File, INDEX_NONE)) //May fail  on encrypted Paks
+
+		if (FCoreDelegates::MountPak.Execute(File, INDEX_NONE)) //May fail on encrypted Paks
 		{
 			bMountedPaks = true;
 		}
@@ -1523,13 +1541,17 @@ bool FDefaultInstallBundleManager::MountPaksInList(TArrayView<FString> Paths, EL
 bool FDefaultInstallBundleManager::UnmountPaksInList(TArrayView<FString> Paths, ELogVerbosity::Type LogVerbosityOverride)
 {
 	if (!FCoreDelegates::OnUnmountPak.IsBound())
+	{
 		return false;
+	}
 
 	bool bUnmountedPaks = false;
 	for (const FString& File : Paths)
 	{
-		if (!File.EndsWith(TEXT(".pak")))
+		if (!File.EndsWith(TEXTVIEW(".pak")))
+		{
 			continue;
+		}
 
 		LOG_INSTALL_BUNDLE_MAN_OVERRIDE(LogVerbosityOverride, Display, TEXT("Unmounting pak file: %s \n"), *File);
 
@@ -1643,6 +1665,7 @@ void FDefaultInstallBundleManager::FinishRequest(FContentRequestRef Request)
 		ResultInfo.Result = Request->Result;
 		ResultInfo.bIsStartup = BundleInfo.bIsStartup;
 		ResultInfo.bContentWasInstalled = Request->bContentWasInstalled;
+		ResultInfo.bContainsChunks = BundleInfo.ContentPaths.bContainsChunks;
 		ResultInfo.OptionalErrorText = Request->OptionalErrorText;
 		ResultInfo.OptionalErrorCode = Request->OptionalErrorCode;
 
@@ -3049,20 +3072,31 @@ TValueOrError<FInstallBundleRequestInfo, EInstallBundleResult> FDefaultInstallBu
 		// Don't request finished bundles
 		if (ActiveQueuedRequest == nullptr)
 		{
+			bool bIsFinished = false;
+
 			// If we canceled a release during an async op, that op could change bundle status when it completes, so enqueue the request to run after the 
 			// canceled release has finished.
 			if (!bCanceledRelease && EnumHasAnyFlags(Flags, EInstallBundleRequestFlags::SkipMount) && GetBundleStatus(*BundleInfo) == EBundleState::NeedsMount)
 			{
 				RetInfo.InfoFlags |= EInstallBundleRequestInfoFlags::SkippedAlreadyUpdatedBundles;
 				LOG_INSTALL_BUNDLE_MAN_OVERRIDE(LogVerbosityOverride, Verbose, TEXT("RequestUpdateContent Bundle %s  - Already Updated"), *BundleInfo->BundleNameString);
-				continue;
+				bIsFinished = true;
 			}
-
 			// No need to check bCanceledRelease here.  Unmounting is not Async so if we canceled it early enough we will remain mounted
-			if (!GetMustWaitForPSOCache(*BundleInfo) && GetBundleStatus(*BundleInfo) == EBundleState::Mounted)
+			else if (!GetMustWaitForPSOCache(*BundleInfo) && GetBundleStatus(*BundleInfo) == EBundleState::Mounted)
 			{
 				RetInfo.InfoFlags |= EInstallBundleRequestInfoFlags::SkippedAlreadyMountedBundles;
 				LOG_INSTALL_BUNDLE_MAN_OVERRIDE(LogVerbosityOverride, Verbose, TEXT("RequestUpdateContent Bundle %s  - Already Mounted"), *BundleInfo->BundleNameString);
+				bIsFinished = true;
+			}
+
+			if (bIsFinished)
+			{
+				FInstallBundleRequestResultInfo& ResultInfo = RetInfo.BundleResults.Emplace_GetRef();
+				ResultInfo.BundleName = BundleName;
+				ResultInfo.Result = EInstallBundleResult::OK;
+				ResultInfo.bIsStartup = BundleInfo->bIsStartup;
+				ResultInfo.bContainsChunks = BundleInfo->ContentPaths.bContainsChunks;
 				continue;
 			}
 		}
@@ -3321,12 +3355,12 @@ void FDefaultInstallBundleManager::CancelAllGetInstallStateRequests(FDelegateHan
 	}
 }
 
-TValueOrError<FInstallBundleRequestInfo, EInstallBundleResult> FDefaultInstallBundleManager::RequestReleaseContent(
+TValueOrError<FInstallBundleReleaseRequestInfo, EInstallBundleResult> FDefaultInstallBundleManager::RequestReleaseContent(
 	TArrayView<const FName> ReleaseNames, EInstallBundleReleaseRequestFlags Flags, TArrayView<const FName> KeepNames /*= TArrayView<const FName>()*/, ELogVerbosity::Type LogVerbosityOverride /*= ELogVerbosity::NoLogging*/)
 {
 	CSV_SCOPED_TIMING_STAT(InstallBundleManager, InstallBundleManager_RequestReleaseContent);
 
-	FInstallBundleRequestInfo RetInfo;
+	FInstallBundleReleaseRequestInfo RetInfo;
 
 	// Check for failing init, this is not recoverable and we can't safely enqueue requests.
 	// bUnrecoverableInitError usually means something is wrong with the build.
diff --git a/Engine/Plugins/Experimental/DefaultInstallBundleManager/Source/Public/DefaultInstallBundleManager.h b/Engine/Plugins/Experimental/DefaultInstallBundleManager/Source/Public/DefaultInstallBundleManager.h
--- a/Engine/Plugins/Experimental/DefaultInstallBundleManager/Source/Public/DefaultInstallBundleManager.h
+++ b/Engine/Plugins/Experimental/DefaultInstallBundleManager/Source/Public/DefaultInstallBundleManager.h
@@ -148,6 +148,7 @@ protected:
 		TArray<FString> ContentPaths;
 		TArray<FString> AdditionalRootDirs;
 		TSet<FString> NonUFSShaderLibPaths;
+		bool bContainsChunks = false;
 	};
 
 	friend struct FBundleInfo;
@@ -629,7 +630,7 @@ public:
 	virtual void CancelAllGetInstallStateRequestsForTag(FName RequestTag) override;
 	virtual void CancelAllGetInstallStateRequests(FDelegateHandle Handle) override;
 
-	virtual TValueOrError<FInstallBundleRequestInfo, EInstallBundleResult> RequestReleaseContent(TArrayView<const FName> ReleaseNames, EInstallBundleReleaseRequestFlags Flags, TArrayView<const FName> KeepNames = TArrayView<const FName>(), ELogVerbosity::Type LogVerbosityOverride = ELogVerbosity::NoLogging) override;
+	virtual TValueOrError<FInstallBundleReleaseRequestInfo, EInstallBundleResult> RequestReleaseContent(TArrayView<const FName> ReleaseNames, EInstallBundleReleaseRequestFlags Flags, TArrayView<const FName> KeepNames = TArrayView<const FName>(), ELogVerbosity::Type LogVerbosityOverride = ELogVerbosity::NoLogging) override;
 	
 	virtual EInstallBundleResult FlushCache(FInstallBundleSourceOrCache SourceOrCache, FInstallBundleManagerFlushCacheCompleteDelegate Callback, ELogVerbosity::Type LogVerbosityOverride = ELogVerbosity::NoLogging) override;
 
diff --git a/Engine/Plugins/Experimental/GameFeatures/Source/GameFeatures/Private/GameFeaturePluginStateMachine.cpp b/Engine/Plugins/Experimental/GameFeatures/Source/GameFeatures/Private/GameFeaturePluginStateMachine.cpp
--- a/Engine/Plugins/Experimental/GameFeatures/Source/GameFeatures/Private/GameFeaturePluginStateMachine.cpp
+++ b/Engine/Plugins/Experimental/GameFeatures/Source/GameFeatures/Private/GameFeaturePluginStateMachine.cpp
@@ -63,6 +63,10 @@ namespace UE::GameFeatures
 		false,
 		TEXT("Enable to use aysnc loading"));
 
+	static TAutoConsoleVariable<bool> CVarAllowForceMonolithicShaderLibrary(TEXT("GameFeaturePlugin.AllowForceMonolithicShaderLibrary"),
+		true,
+		TEXT("Enable to force only searching for monolithic shader libs when possible"));
+
 	#define GAME_FEATURE_PLUGIN_STATE_TO_STRING(inEnum, inText) case EGameFeaturePluginState::inEnum: return TEXT(#inEnum);
 	FString ToString(EGameFeaturePluginState InType)
 	{
@@ -1083,7 +1087,7 @@ struct FBaseDataReleaseGameFeaturePluginState : public FGameFeaturePluginState
 		const TArray<FName>& InstallBundles = StateProperties.ProtocolMetadata.GetSubtype<FInstallBundlePluginProtocolMetaData>().InstallBundles;
 
 		EInstallBundleReleaseRequestFlags ReleaseFlags = GetReleaseRequestFlags();
-		TValueOrError<FInstallBundleRequestInfo, EInstallBundleResult> MaybeRequestInfo = BundleManager->RequestReleaseContent(InstallBundles, ReleaseFlags);
+		TValueOrError<FInstallBundleReleaseRequestInfo, EInstallBundleResult> MaybeRequestInfo = BundleManager->RequestReleaseContent(InstallBundles, ReleaseFlags);
 
 		if (MaybeRequestInfo.HasError())
 		{
@@ -1094,7 +1098,7 @@ struct FBaseDataReleaseGameFeaturePluginState : public FGameFeaturePluginState
 			return;
 		}
 
-		FInstallBundleRequestInfo RequestInfo = MaybeRequestInfo.StealValue();
+		FInstallBundleReleaseRequestInfo RequestInfo = MaybeRequestInfo.StealValue();
 
 		if (EnumHasAnyFlags(RequestInfo.InfoFlags, EInstallBundleRequestInfoFlags::SkippedUnknownBundles))
 		{
@@ -1730,7 +1734,7 @@ struct FGameFeaturePluginState_Unmounting : public FGameFeaturePluginState
 		//Make sure we don't remove files here early, that should only be done in Uninstalling
 		ReleaseFlags &= ~(EInstallBundleReleaseRequestFlags::RemoveFilesIfPossible);
 
-		TValueOrError<FInstallBundleRequestInfo, EInstallBundleResult> MaybeRequestInfo = BundleManager->RequestReleaseContent(InstallBundles, ReleaseFlags);
+		TValueOrError<FInstallBundleReleaseRequestInfo, EInstallBundleResult> MaybeRequestInfo = BundleManager->RequestReleaseContent(InstallBundles, ReleaseFlags);
 
 		if (MaybeRequestInfo.HasError())
 		{
@@ -1740,7 +1744,7 @@ struct FGameFeaturePluginState_Unmounting : public FGameFeaturePluginState
 			return;
 		}
 
-		FInstallBundleRequestInfo RequestInfo = MaybeRequestInfo.StealValue();
+		FInstallBundleReleaseRequestInfo RequestInfo = MaybeRequestInfo.StealValue();
 
 		if (EnumHasAnyFlags(RequestInfo.InfoFlags, EInstallBundleRequestInfoFlags::SkippedUnknownBundles))
 		{
@@ -1858,6 +1862,7 @@ struct FGameFeaturePluginState_Mounting : public FGameFeaturePluginState
 	ESubState StartedSubStates = ESubState::None;
 	ESubState CompletedSubStates = ESubState::None;
 	bool bCheckedRealtimeMode = false;
+	bool bForceMonolithicShaderLibrary = false;
 
 	void OnInstallBundleCompleted(FInstallBundleRequestResultInfo BundleResult)
 	{
@@ -1880,6 +1885,11 @@ struct FGameFeaturePluginState_Mounting : public FGameFeaturePluginState
 			}
 		}
 
+		if (bForceMonolithicShaderLibrary && BundleResult.bContainsChunks)
+		{
+			bForceMonolithicShaderLibrary = false;
+		}
+
 		if (PendingBundles.IsEmpty())
 		{
 			IInstallBundleManager::InstallBundleCompleteDelegate.RemoveAll(this);
@@ -1932,11 +1942,15 @@ struct FGameFeaturePluginState_Mounting : public FGameFeaturePluginState
 		StartedSubStates = ESubState::None;
 		CompletedSubStates = ESubState::None;
 		bCheckedRealtimeMode = false;
+		bForceMonolithicShaderLibrary = false;
 
 		if (StateProperties.GetPluginProtocol() != EGameFeaturePluginProtocol::InstallBundle)
 		{
 			return;
 		}
+
+		// Assume monolithic shader, will be set to false if chunks are detected
+		bForceMonolithicShaderLibrary = UE::GameFeatures::CVarAllowForceMonolithicShaderLibrary.GetValueOnGameThread();
 		
 		TSharedPtr<IInstallBundleManager> BundleManager = IInstallBundleManager::GetPlatformInstallBundleManager();
 
@@ -1980,6 +1994,14 @@ struct FGameFeaturePluginState_Mounting : public FGameFeaturePluginState
 				PakFileMountedDelegateHandle = FCoreDelegates::GetOnPakFileMounted2().AddRaw(this, &FGameFeaturePluginState_Mounting::OnPakFileMounted);
 			}
 		}
+
+		for (const FInstallBundleRequestResultInfo& BundleResult : RequestInfo.BundleResults)
+		{
+			if (bForceMonolithicShaderLibrary && BundleResult.bContainsChunks)
+			{
+				bForceMonolithicShaderLibrary = false;
+			}
+		}
 	}
 
 	void UpdateState_MountPlugin()
@@ -2000,9 +2022,11 @@ struct FGameFeaturePluginState_Mounting : public FGameFeaturePluginState
 		}
 
 		// Pre-mount
+		bool bOpenPluginShaderLibrary = true;
 		{
 			FGameFeaturePreMountingContext Context;
 			UGameFeaturesSubsystem::Get().OnGameFeaturePreMounting(StateProperties.PluginName, StateProperties.PluginIdentifier, Context);
+			bOpenPluginShaderLibrary = Context.bOpenPluginShaderLibrary;
 		}
 
 		checkf(!StateProperties.PluginInstalledFilename.IsEmpty(), TEXT("PluginInstalledFilename must be set by the Mounting. PluginURL: %s"), *StateProperties.PluginIdentifier.GetFullPluginURL());
@@ -2038,25 +2062,33 @@ struct FGameFeaturePluginState_Mounting : public FGameFeaturePluginState
 			return;
 		}
 
+		if (bOpenPluginShaderLibrary)
+		{
+			// We want to control opening the shader lib
+			FShaderCodeLibrary::DontOpenPluginShaderLibraryOnMount(StateProperties.PluginName);
+		}
+
 		if (!UseAsyncLoading())
 		{
 			verify(IPluginManager::Get().MountExplicitlyLoadedPlugin(StateProperties.PluginName));
+			if (bOpenPluginShaderLibrary)
+			{
+				TSharedPtr<IPlugin> Plugin = IPluginManager::Get().FindPlugin(StateProperties.PluginName);
+				FShaderCodeLibrary::OpenPluginShaderLibrary(*Plugin, bForceMonolithicShaderLibrary);
+			}
 			CompletedSubStates |= ESubState::MountPlugin;
 			return;
 		}
 
-		// We want to make sure opening the shader lib happens async
-		FShaderCodeLibrary::DontOpenPluginShaderLibraryOnMount(StateProperties.PluginName);
-
 		verify(IPluginManager::Get().MountExplicitlyLoadedPlugin(StateProperties.PluginName));
 
 		// Now load the shader lib in the background
 		TSharedPtr<IPlugin> Plugin = IPluginManager::Get().FindPlugin(StateProperties.PluginName);
-		if (Plugin->CanContainContent() && Plugin->IsEnabled()) // TODO: possibly skip this if there is no shaderlib (need to figure out the file name)
+		if (bOpenPluginShaderLibrary && Plugin->CanContainContent() && Plugin->IsEnabled())
 		{
 			UE::Tasks::Launch(UE_SOURCE_LOCATION, [this, PluginName=Plugin->GetName(), PluginDir=Plugin->GetContentDir()]
 			{
-				FShaderCodeLibrary::OpenLibrary(PluginName, PluginDir);
+				FShaderCodeLibrary::OpenLibrary(PluginName, PluginDir, bForceMonolithicShaderLibrary);
 
 				ExecuteOnGameThread(UE_SOURCE_LOCATION, [this]
 				{
diff --git a/Engine/Plugins/Experimental/GameFeatures/Source/GameFeatures/Public/GameFeaturesSubsystem.h b/Engine/Plugins/Experimental/GameFeatures/Source/GameFeatures/Public/GameFeaturesSubsystem.h
--- a/Engine/Plugins/Experimental/GameFeatures/Source/GameFeatures/Public/GameFeaturesSubsystem.h
+++ b/Engine/Plugins/Experimental/GameFeatures/Source/GameFeatures/Public/GameFeaturesSubsystem.h
@@ -160,7 +160,7 @@ private:
 struct FGameFeaturePreMountingContext : public FGameFeatureStateChangeContext
 {
 public:
-	//@TODO: Add rules specific to pre-mounting when required
+	bool bOpenPluginShaderLibrary = true;
 
 private:
 
diff --git a/Engine/Source/Runtime/Core/Private/GenericPlatform/GenericPlatformMisc.cpp b/Engine/Source/Runtime/Core/Private/GenericPlatform/GenericPlatformMisc.cpp
--- a/Engine/Source/Runtime/Core/Private/GenericPlatform/GenericPlatformMisc.cpp
+++ b/Engine/Source/Runtime/Core/Private/GenericPlatform/GenericPlatformMisc.cpp
@@ -16,6 +16,7 @@
 #include "Misc/Parse.h"
 #include "Misc/CommandLine.h"
 #include "Misc/Paths.h"
+#include "Misc/PathViews.h"
 #include "Misc/FileHelper.h"
 #include "Internationalization/Text.h"
 #include "Internationalization/Internationalization.h"
@@ -1982,29 +1983,27 @@ void FGenericPlatformMisc::ParseChunkIdPakchunkIndexMapping(TArray<FString> Chun
 	}
 }
 
-int32 FGenericPlatformMisc::GetPakchunkIndexFromPakFile(const FString& InFilename)
+int32 FGenericPlatformMisc::GetPakchunkIndexFromPakFile(FStringView InFilename)
 {
-	FString ChunkIdentifier(TEXT("pakchunk"));
-	FString BaseFilename = FPaths::GetBaseFilename(InFilename);
+	FStringView ChunkIdentifier(TEXTVIEW("pakchunk"));
+	FStringView BaseFilename = FPathViews::GetBaseFilename(InFilename);
 	int32 ChunkNumber = INDEX_NONE;
 
 	if (BaseFilename.StartsWith(ChunkIdentifier))
 	{
-		int32 StartOfNumber = ChunkIdentifier.Len();
+		const int32 StartOfNumber = ChunkIdentifier.Len();
 		int32 DigitCount = 0;
-		if (FChar::IsDigit(BaseFilename[StartOfNumber]))
+		
+		while ((DigitCount + StartOfNumber) < BaseFilename.Len() && FChar::IsDigit(BaseFilename[StartOfNumber + DigitCount]))
 		{
-			while ((DigitCount + StartOfNumber) < BaseFilename.Len() && FChar::IsDigit(BaseFilename[StartOfNumber + DigitCount]))
-			{
-				DigitCount++;
-			}
+			DigitCount++;
+		}
 
-			if ((StartOfNumber + DigitCount) < BaseFilename.Len())
-			{
-				FString ChunkNumberString = BaseFilename.Mid(StartOfNumber, DigitCount);
-				check(ChunkNumberString.IsNumeric());
-				TTypeFromString<int32>::FromString(ChunkNumber, *ChunkNumberString);
-			}
+		if (DigitCount > 0 && (StartOfNumber + DigitCount) < BaseFilename.Len())
+		{
+			// FromString can't take a view
+			TStringBuilder<16> ChunkNumberString = WriteToString<16>(BaseFilename.Mid(StartOfNumber, DigitCount));
+			TTypeFromString<int32>::FromString(ChunkNumber, *ChunkNumberString);
 		}
 	}
 
diff --git a/Engine/Source/Runtime/Core/Public/GenericPlatform/GenericPlatformMisc.h b/Engine/Source/Runtime/Core/Public/GenericPlatform/GenericPlatformMisc.h
--- a/Engine/Source/Runtime/Core/Public/GenericPlatform/GenericPlatformMisc.h
+++ b/Engine/Source/Runtime/Core/Public/GenericPlatform/GenericPlatformMisc.h
@@ -1840,7 +1840,7 @@ public:
 		return PakchunkIndex;
 	}
 
-	static CORE_API int32 GetPakchunkIndexFromPakFile(const FString& InFilename);
+	static CORE_API int32 GetPakchunkIndexFromPakFile(FStringView InFilename);
 
 	FORCEINLINE static bool Expand16BitIndicesTo32BitOnLoad()
 	{
diff --git a/Engine/Source/Runtime/InstallBundleManager/Private/InstallBundleManagerInterface.cpp b/Engine/Source/Runtime/InstallBundleManager/Private/InstallBundleManagerInterface.cpp
--- a/Engine/Source/Runtime/InstallBundleManager/Private/InstallBundleManagerInterface.cpp
+++ b/Engine/Source/Runtime/InstallBundleManager/Private/InstallBundleManagerInterface.cpp
@@ -74,7 +74,7 @@ TValueOrError<FInstallBundleCombinedInstallState, EInstallBundleResult> IInstall
 	return GetInstallStateSynchronous(MakeArrayView(&BundleName, 1), bAddDependencies);
 }
 
-TValueOrError<FInstallBundleRequestInfo, EInstallBundleResult> IInstallBundleManager::RequestReleaseContent(FName ReleaseName, EInstallBundleReleaseRequestFlags Flags, TArrayView<const FName> KeepNames /*= TArrayView<const FName>()*/, ELogVerbosity::Type LogVerbosityOverride /*= ELogVerbosity::NoLogging*/)
+TValueOrError<FInstallBundleReleaseRequestInfo, EInstallBundleResult> IInstallBundleManager::RequestReleaseContent(FName ReleaseName, EInstallBundleReleaseRequestFlags Flags, TArrayView<const FName> KeepNames /*= TArrayView<const FName>()*/, ELogVerbosity::Type LogVerbosityOverride /*= ELogVerbosity::NoLogging*/)
 {
 	return RequestReleaseContent(MakeArrayView(&ReleaseName, 1), Flags, KeepNames, LogVerbosityOverride);
 }
diff --git a/Engine/Source/Runtime/InstallBundleManager/Public/InstallBundleManagerInterface.h b/Engine/Source/Runtime/InstallBundleManager/Public/InstallBundleManagerInterface.h
--- a/Engine/Source/Runtime/InstallBundleManager/Public/InstallBundleManagerInterface.h
+++ b/Engine/Source/Runtime/InstallBundleManager/Public/InstallBundleManagerInterface.h
@@ -48,6 +48,7 @@ struct FInstallBundleRequestResultInfo
 	FName BundleName;
 	EInstallBundleResult Result = EInstallBundleResult::OK;
 	bool bIsStartup = false;
+	bool bContainsChunks = false;
 	bool bContentWasInstalled = false;
 
 	// Currently, these just forward BPT Error info
@@ -55,16 +56,29 @@ struct FInstallBundleRequestResultInfo
 	FString OptionalErrorCode;
 };
 
-struct FInstallBundlePauseInfo
+struct FInstallBundleReleaseRequestResultInfo
 {
 	FName BundleName;
-	EInstallBundlePauseFlags PauseFlags = EInstallBundlePauseFlags::None;
+	EInstallBundleReleaseResult Result = EInstallBundleReleaseResult::OK; 
 };
 
-struct FInstallBundleReleaseRequestResultInfo
+struct FInstallBundleRequestInfo
+{
+	EInstallBundleRequestInfoFlags InfoFlags = EInstallBundleRequestInfoFlags::None;
+	TArray<FName> BundlesEnqueued;
+	TArray<FInstallBundleRequestResultInfo> BundleResults;
+};
+
+struct FInstallBundleReleaseRequestInfo
+{
+	EInstallBundleRequestInfoFlags InfoFlags = EInstallBundleRequestInfoFlags::None;
+	TArray<FName> BundlesEnqueued;
+};
+
+struct FInstallBundlePauseInfo
 {
 	FName BundleName;
-	EInstallBundleReleaseResult Result = EInstallBundleReleaseResult::OK; 
+	EInstallBundlePauseFlags PauseFlags = EInstallBundlePauseFlags::None;
 };
 
 enum class EInstallBundleManagerInitErrorHandlerResult
@@ -136,8 +150,8 @@ public:
 	virtual void CancelAllGetInstallStateRequestsForTag(FName RequestTag) = 0;
 	virtual void CancelAllGetInstallStateRequests(FDelegateHandle Handle) = 0;
 
-	INSTALLBUNDLEMANAGER_API TValueOrError<FInstallBundleRequestInfo, EInstallBundleResult> RequestReleaseContent(FName ReleaseName, EInstallBundleReleaseRequestFlags Flags, TArrayView<const FName> KeepNames = TArrayView<const FName>(), ELogVerbosity::Type LogVerbosityOverride = ELogVerbosity::NoLogging);
-	virtual TValueOrError<FInstallBundleRequestInfo, EInstallBundleResult> RequestReleaseContent(TArrayView<const FName> ReleaseNames, EInstallBundleReleaseRequestFlags Flags, TArrayView<const FName> KeepNames = TArrayView<const FName>(), ELogVerbosity::Type LogVerbosityOverride = ELogVerbosity::NoLogging) = 0;
+	INSTALLBUNDLEMANAGER_API TValueOrError<FInstallBundleReleaseRequestInfo, EInstallBundleResult> RequestReleaseContent(FName ReleaseName, EInstallBundleReleaseRequestFlags Flags, TArrayView<const FName> KeepNames = TArrayView<const FName>(), ELogVerbosity::Type LogVerbosityOverride = ELogVerbosity::NoLogging);
+	virtual TValueOrError<FInstallBundleReleaseRequestInfo, EInstallBundleResult> RequestReleaseContent(TArrayView<const FName> ReleaseNames, EInstallBundleReleaseRequestFlags Flags, TArrayView<const FName> KeepNames = TArrayView<const FName>(), ELogVerbosity::Type LogVerbosityOverride = ELogVerbosity::NoLogging) = 0;
 
 	INSTALLBUNDLEMANAGER_API EInstallBundleResult FlushCache(FInstallBundleManagerFlushCacheCompleteDelegate Callback, ELogVerbosity::Type LogVerbosityOverride = ELogVerbosity::NoLogging);
 	virtual EInstallBundleResult FlushCache(FInstallBundleSourceOrCache SourceOrCache, FInstallBundleManagerFlushCacheCompleteDelegate Callback, ELogVerbosity::Type LogVerbosityOverride = ELogVerbosity::NoLogging) = 0;
diff --git a/Engine/Source/Runtime/InstallBundleManager/Public/InstallBundleTypes.h b/Engine/Source/Runtime/InstallBundleManager/Public/InstallBundleTypes.h
--- a/Engine/Source/Runtime/InstallBundleManager/Public/InstallBundleTypes.h
+++ b/Engine/Source/Runtime/InstallBundleManager/Public/InstallBundleTypes.h
@@ -173,12 +173,6 @@ enum class EInstallBundleReleaseRequestFlags : uint32
 };
 ENUM_CLASS_FLAGS(EInstallBundleReleaseRequestFlags)
 
-struct FInstallBundleRequestInfo
-{
-	EInstallBundleRequestInfoFlags InfoFlags = EInstallBundleRequestInfoFlags::None;
-	TArray<FName> BundlesEnqueued;
-};
-
 enum class EInstallBundlePauseFlags : uint32
 {
 	None = 0,
diff --git a/Engine/Source/Runtime/NullInstallBundleManager/Private/NullInstallBundleManager.cpp b/Engine/Source/Runtime/NullInstallBundleManager/Private/NullInstallBundleManager.cpp
--- a/Engine/Source/Runtime/NullInstallBundleManager/Private/NullInstallBundleManager.cpp
+++ b/Engine/Source/Runtime/NullInstallBundleManager/Private/NullInstallBundleManager.cpp
@@ -73,9 +73,9 @@ class FNullInstallBundleManager : public IInstallBundleManager
 	{
 	}
 
-	virtual TValueOrError<FInstallBundleRequestInfo, EInstallBundleResult> RequestReleaseContent(TArrayView<const FName> ReleaseNames, EInstallBundleReleaseRequestFlags Flags, TArrayView<const FName> KeepNames = TArrayView<const FName>(), ELogVerbosity::Type LogVerbosityOverride = ELogVerbosity::NoLogging) override
+	virtual TValueOrError<FInstallBundleReleaseRequestInfo, EInstallBundleResult> RequestReleaseContent(TArrayView<const FName> ReleaseNames, EInstallBundleReleaseRequestFlags Flags, TArrayView<const FName> KeepNames = TArrayView<const FName>(), ELogVerbosity::Type LogVerbosityOverride = ELogVerbosity::NoLogging) override
 	{
-		return MakeValue(FInstallBundleRequestInfo());
+		return MakeValue(FInstallBundleReleaseRequestInfo());
 	}
 
 	virtual EInstallBundleResult FlushCache(FInstallBundleSourceOrCache SourceOrCache, FInstallBundleManagerFlushCacheCompleteDelegate Callback, ELogVerbosity::Type LogVerbosityOverride = ELogVerbosity::NoLogging) override
diff --git a/Engine/Source/Runtime/PakFile/Private/IPlatformFilePak.cpp b/Engine/Source/Runtime/PakFile/Private/IPlatformFilePak.cpp
--- a/Engine/Source/Runtime/PakFile/Private/IPlatformFilePak.cpp
+++ b/Engine/Source/Runtime/PakFile/Private/IPlatformFilePak.cpp
@@ -88,31 +88,7 @@ static FString GMountStartupPaksWildCard = TEXT(MOUNT_STARTUP_PAKS_WILDCARD);
 
 int32 GetPakchunkIndexFromPakFile(const FString& InFilename)
 {
-	FString ChunkIdentifier(TEXT("pakchunk"));
-	FString BaseFilename = FPaths::GetBaseFilename(InFilename);
-	int32 ChunkNumber = INDEX_NONE;
-
-	if (BaseFilename.StartsWith(ChunkIdentifier))
-	{
-		int32 StartOfNumber = ChunkIdentifier.Len();
-		int32 DigitCount = 0;
-		if (FChar::IsDigit(BaseFilename[StartOfNumber]))
-		{
-			while ((DigitCount + StartOfNumber) < BaseFilename.Len() && FChar::IsDigit(BaseFilename[StartOfNumber + DigitCount]))
-			{
-				DigitCount++;
-			}
-
-			if ((StartOfNumber + DigitCount) < BaseFilename.Len())
-			{
-				FString ChunkNumberString = BaseFilename.Mid(StartOfNumber, DigitCount);
-				check(ChunkNumberString.IsNumeric());
-				TTypeFromString<int32>::FromString(ChunkNumber, *ChunkNumberString);
-			}
-		}
-	}
-
-	return ChunkNumber;
+	return FGenericPlatformMisc::GetPakchunkIndexFromPakFile(InFilename);
 }
 
 // Registered encryption key cache
@@ -8491,7 +8467,7 @@ void FPakPlatformFile::OptimizeMemoryUsageForMountedPaks()
 }
 
 
-bool FPakPlatformFile::Mount(const TCHAR* InPakFilename, uint32 PakOrder, const TCHAR* InPath /*= NULL*/, bool bLoadIndex /*= true*/)
+bool FPakPlatformFile::Mount(const TCHAR* InPakFilename, uint32 PakOrder, const TCHAR* InPath /*= nullptr*/, bool bLoadIndex /*= true*/, FPakListEntry* OutPakListEntry /*= nullptr*/)
 {
 	LLM_SCOPE(ELLMTag::FileSystem);
 	bool bPakSuccess = false;
@@ -8503,7 +8479,7 @@ bool FPakPlatformFile::Mount(const TCHAR* InPakFilename, uint32 PakOrder, const
 		{
 			if (!Pak->GetInfo().EncryptionKeyGuid.IsValid() || GetRegisteredEncryptionKeys().HasKey(Pak->GetInfo().EncryptionKeyGuid))
 			{
-				if (InPath != NULL)
+				if (InPath != nullptr)
 				{
 					Pak->SetMountPoint(InPath);
 				}
@@ -8544,6 +8520,11 @@ bool FPakPlatformFile::Mount(const TCHAR* InPakFilename, uint32 PakOrder, const
 					Pak->SetIsMounted(true);
 					PakFiles.Add(Entry);
 					PakFiles.StableSort();
+
+					if (OutPakListEntry)
+					{
+						*OutPakListEntry = MoveTemp(Entry);
+					}
 				}
 				bPakSuccess = true;
 			}
@@ -8824,7 +8805,7 @@ int32 FPakPlatformFile::MountAllPakFiles(const TArray<FString>& PakFolders, cons
 		GetMountedPaks(ExistingPaks);
 		TSet<FString> ExistingPaksFileName;
 		// Find the single pak we just mounted
-		for (auto Pak : ExistingPaks)
+		for (const FPakListEntry& Pak : ExistingPaks)
 		{
 			ExistingPaksFileName.Add(Pak.PakFile->GetFilename());
 		}
@@ -8892,20 +8873,10 @@ IPakFile* FPakPlatformFile::HandleMountPakDelegate(const FString& PakFilePath, i
 		PakOrder = GetPakOrderFromPakFilePath(PakFilePath);
 	}
 	
-	if (Mount(*PakFilePath, PakOrder))
+	FPakListEntry Pak;
+	if (Mount(*PakFilePath, PakOrder, nullptr, true, &Pak))
 	{
-		TArray<FPakListEntry> Paks;
-		GetMountedPaks(Paks);
-		// Find the single pak we just mounted
-		for (auto Pak : Paks)
-		{
-			if (PakFilePath == Pak.PakFile->GetFilename())
-			{
-				return Pak.PakFile;
-			}
-		}
-
-		check(false); // Should have found the pak since Mount returned true
+		return Pak.PakFile;
 	}
 	return nullptr;
 }
diff --git a/Engine/Source/Runtime/PakFile/Public/IPlatformFilePak.h b/Engine/Source/Runtime/PakFile/Public/IPlatformFilePak.h
--- a/Engine/Source/Runtime/PakFile/Public/IPlatformFilePak.h
+++ b/Engine/Source/Runtime/PakFile/Public/IPlatformFilePak.h
@@ -2110,6 +2110,19 @@ public:
 		}
 	}
 
+	/**
+	* Get a list of all pak files which have been successfully mounted
+	*/
+	FORCEINLINE void GetMountedPakFilenames(TSet<FString>& PakFilenames)
+	{
+		FScopeLock ScopedLock(&PakListCritical);
+		PakFilenames.Empty(PakFiles.Num());
+		for (FPakListEntry& Entry : PakFiles)
+		{
+			PakFilenames.Add(Entry.PakFile->GetFilename());
+		}
+	}
+
 	/**
 	 * Checks if pak files exist in any of the known pak file locations.
 	 */
@@ -2177,7 +2190,7 @@ public:
 	 * @param InPakFilename Pak filename.
 	 * @param InPath Path to mount the pak at.
 	 */
-	PAKFILE_API bool Mount(const TCHAR* InPakFilename, uint32 PakOrder, const TCHAR* InPath = NULL, bool bLoadIndex = true);
+	PAKFILE_API bool Mount(const TCHAR* InPakFilename, uint32 PakOrder, const TCHAR* InPath = nullptr, bool bLoadIndex = true, FPakListEntry* OutPakListEntry = nullptr);
 
 	PAKFILE_API bool Unmount(const TCHAR* InPakFilename);
 
