Platform chunk bundle source wrapper will re-trigger the completion callback if it finds the name chunk already installed when requested

 - this is to support cases where multiple bundles are mapped to the same named chunk

#rnx

diff --git a/Engine/Plugins/Experimental/DefaultInstallBundleManager/Source/Private/InstallBundleSourcePlatformChunkInstall.cpp b/Engine/Plugins/Experimental/DefaultInstallBundleManager/Source/Private/InstallBundleSourcePlatformChunkInstall.cpp
--- a/Engine/Plugins/Experimental/DefaultInstallBundleManager/Source/Private/InstallBundleSourcePlatformChunkInstall.cpp
+++ b/Engine/Plugins/Experimental/DefaultInstallBundleManager/Source/Private/InstallBundleSourcePlatformChunkInstall.cpp
@@ -456,8 +456,15 @@ void FInstallBundleSourcePlatformChunkInstall::RequestUpdateContent(FRequestUpda
 	}
 	else if (PlatformChunkInstall->GetNamedChunkLocation(BundleInfo->NamedChunk) == EChunkLocation::LocalFast)
 	{
-		LOG_SOURCE_CHUNKINSTALL_OVERRIDE(Context.LogVerbosityOverride, Error, TEXT("Bundle %s is already installed (named chunk: %s)"), *Context.BundleName.ToString(), *BundleInfo->NamedChunk.ToString());
-		bFailed = true;
+		LOG_SOURCE_CHUNKINSTALL_OVERRIDE(Context.LogVerbosityOverride, Display, TEXT("Bundle request %s finished. (named chunk %s already installed)"), *Context.BundleName.ToString(), *BundleInfo->NamedChunk.ToString() );
+
+		// send the completion callback immediately
+		FInstallBundleSourceUpdateContentResultInfo ResultInfo;
+		ResultInfo.BundleName = Context.BundleName;
+		ResultInfo.Result = EInstallBundleResult::OK;
+		ResultInfo.ContentPaths = BundleInfo->FilePaths;
+		ResultInfo.bContentWasInstalled = (BundleInfo->FilePaths.Num() > 0);
+		Context.CompleteCallback.ExecuteIfBound(AsShared(), MoveTemp(ResultInfo));
 	}
 	else
 	{
