Remove GInstallBundleManagerIni

diff --git a/Engine/Plugins/Experimental/DefaultInstallBundleManager/Source/Private/DefaultInstallBundleManager.cpp b/Engine/Plugins/Experimental/DefaultInstallBundleManager/Source/Private/DefaultInstallBundleManager.cpp
--- a/Engine/Plugins/Experimental/DefaultInstallBundleManager/Source/Private/DefaultInstallBundleManager.cpp
+++ b/Engine/Plugins/Experimental/DefaultInstallBundleManager/Source/Private/DefaultInstallBundleManager.cpp
@@ -185,16 +185,12 @@ void LoadKeychainFromInI(FKeyChain& OutCryptoSettings)
 }
 #endif // WITH_EDITOR
 
-FDefaultInstallBundleManager::FDefaultInstallBundleManager(const TCHAR* InConfigBaseName, FInstallBundleSourceFactoryFunction InInstallBundleSourceFactory)
+FDefaultInstallBundleManager::FDefaultInstallBundleManager(FInstallBundleSourceFactoryFunction InInstallBundleSourceFactory)
 	: InstallBundleSourceFactory(InInstallBundleSourceFactory ? InInstallBundleSourceFactory : InstallBundleManagerUtil::MakeBundleSource )
 	, PersistentStats(MakeShared<InstallBundleManagerUtil::FPersistentStatContainer>())
 	, AnalyticsProvider(nullptr)
 	, StatsMap(MakeShared<InstallBundleUtil::FContentRequestStatsMap>())
 {
-	// Init Config Settings
-	const TCHAR* DefaultConfigBaseName = TEXT("InstallBundleManager");
-	FConfigContext::ReadIntoGConfig().Load(InConfigBaseName ? InConfigBaseName : DefaultConfigBaseName, GInstallBundleManagerIni);
-
 #if WITH_EDITOR
 	// -UsePaks needs to be specified on the command line for valid pak to be created.
 	// To support mounting pak files in the editor binary add the encryption key.
diff --git a/Engine/Plugins/Experimental/DefaultInstallBundleManager/Source/Private/InstallBundleManagerUtil.cpp b/Engine/Plugins/Experimental/DefaultInstallBundleManager/Source/Private/InstallBundleManagerUtil.cpp
--- a/Engine/Plugins/Experimental/DefaultInstallBundleManager/Source/Private/InstallBundleManagerUtil.cpp
+++ b/Engine/Plugins/Experimental/DefaultInstallBundleManager/Source/Private/InstallBundleManagerUtil.cpp
@@ -14,8 +14,6 @@
 	#define WITH_PLATFORM_INSTALL_BUNDLE_SOURCE 0
 #endif
 
-DEFAULTINSTALLBUNDLEMANAGER_API FString GInstallBundleManagerIni;
-
 DEFINE_LOG_CATEGORY(LogDefaultInstallBundleManager);
 
 // Json writer subclass to allow us to avoid using a SharedPtr to write basic Json.
diff --git a/Engine/Plugins/Experimental/DefaultInstallBundleManager/Source/Public/DefaultInstallBundleManager.h b/Engine/Plugins/Experimental/DefaultInstallBundleManager/Source/Public/DefaultInstallBundleManager.h
--- a/Engine/Plugins/Experimental/DefaultInstallBundleManager/Source/Public/DefaultInstallBundleManager.h
+++ b/Engine/Plugins/Experimental/DefaultInstallBundleManager/Source/Public/DefaultInstallBundleManager.h
@@ -474,7 +474,12 @@ protected:
 public:
 	typedef TFunction<TSharedPtr<IInstallBundleSource>(EInstallBundleSourceType)> FInstallBundleSourceFactoryFunction;
 
-	FDefaultInstallBundleManager(const TCHAR* InConfigBaseName = nullptr, FInstallBundleSourceFactoryFunction InBundleSourceFactory = nullptr);
+	UE_DEPRECATED(5.4, "GInstallBundleManagerIni is deprecated, use InstallBundle.ini hierarchy instead.")
+	FDefaultInstallBundleManager(const TCHAR* InConfigBaseName, FInstallBundleSourceFactoryFunction InBundleSourceFactory = nullptr)
+		: FDefaultInstallBundleManager(InBundleSourceFactory)
+	{}
+
+	FDefaultInstallBundleManager(FInstallBundleSourceFactoryFunction InBundleSourceFactory = nullptr);
 	FDefaultInstallBundleManager(const FDefaultInstallBundleManager& Other) = delete;
 	FDefaultInstallBundleManager& operator=(const FDefaultInstallBundleManager& Other) = delete;
 
diff --git a/Engine/Plugins/Experimental/DefaultInstallBundleManager/Source/Public/InstallBundleManagerUtil.h b/Engine/Plugins/Experimental/DefaultInstallBundleManager/Source/Public/InstallBundleManagerUtil.h
--- a/Engine/Plugins/Experimental/DefaultInstallBundleManager/Source/Public/InstallBundleManagerUtil.h
+++ b/Engine/Plugins/Experimental/DefaultInstallBundleManager/Source/Public/InstallBundleManagerUtil.h
@@ -14,8 +14,6 @@ struct FAnalyticsEventAttribute;
 class IAnalyticsProviderET;
 struct FInstallBundleCacheStats;
 
-extern DEFAULTINSTALLBUNDLEMANAGER_API FString GInstallBundleManagerIni;
-
 DECLARE_LOG_CATEGORY_EXTERN(LogDefaultInstallBundleManager, Display, All);
 
 namespace InstallBundleManagerUtil
