Added support for duplicated builtin content in cod builds.

diff --git a/Engine/Plugins/Experimental/DefaultInstallBundleManager/Source/Private/DefaultInstallBundleManager.cpp b/Engine/Plugins/Experimental/DefaultInstallBundleManager/Source/Private/DefaultInstallBundleManager.cpp
--- a/Engine/Plugins/Experimental/DefaultInstallBundleManager/Source/Private/DefaultInstallBundleManager.cpp
+++ b/Engine/Plugins/Experimental/DefaultInstallBundleManager/Source/Private/DefaultInstallBundleManager.cpp
@@ -882,10 +882,22 @@ void FDefaultInstallBundleManager::TryReserveCache(FContentRequestRef Request)
 	bool bMustWaitForCacheEvict = false;
 	TMap<FName, EInstallBundleCacheReserveResult> ReserveResults;
 	ReserveResults.Reserve(BundleCaches.Num());
-	for (const TPair<FName, TSharedRef<FInstallBundleCache>>& Pair : BundleCaches)
+
+	TSet<FName> EnabledBundleCaches;
+	for (const TPair<EInstallBundleSourceType, TSharedPtr<IInstallBundleSource>>& Pair : GetEnabledBundleSourcesForRequest(Request))
+	{
+		FName* BundleCacheName = BundleSourceCaches.Find(Pair.Key);
+		if (BundleCacheName != nullptr)
+		{
+			EnabledBundleCaches.Add(*BundleCacheName);
+		}
+	}
+	// we will try reserve cache space for only enabled bundle sources.  Each bundle source knows how much cache it should require for the bundle. 
+	for (const FName& BundleCacheName : EnabledBundleCaches)
 	{
-		FInstallBundleCacheReserveResult Result = Pair.Value->Reserve(Request->BundleName);
-		ReserveResults.Add(Pair.Key, Result.Result);
+		TSharedRef<FInstallBundleCache> BundleCache = BundleCaches.FindRef(BundleCacheName);
+		FInstallBundleCacheReserveResult Result = BundleCache->Reserve(Request->BundleName);
+		ReserveResults.Add(BundleCacheName, Result.Result);
 		switch (Result.Result)
 		{
 		case EInstallBundleCacheReserveResult::Fail_CacheFull:
@@ -1123,13 +1135,19 @@ void FDefaultInstallBundleManager::CacheEvictionComplete(TSharedRef<IInstallBund
 	}
 }
 
+TMap<EInstallBundleSourceType, TSharedPtr<IInstallBundleSource>> FDefaultInstallBundleManager::GetEnabledBundleSourcesForRequest(FContentRequestRef Request) const
+{
+	return BundleSources;
+}
+
 void FDefaultInstallBundleManager::UpdateBundleSources(FContentRequestRef Request)
 {
 	StatsBegin(Request->BundleName, EContentRequestState::UpdatingBundleSources);
 
 	Request->StepResult = EContentRequestStepResult::Waiting;
-
-	for (const TPair<EInstallBundleSourceType, TSharedPtr<IInstallBundleSource>>& Pair : BundleSources)
+	TMap<EInstallBundleSourceType, TSharedPtr<IInstallBundleSource>> EnabledBundleSources = GetEnabledBundleSourcesForRequest(Request);
+	Request->RequiredSourceRequestResultsCount = EnabledBundleSources.Num();
+	for (const TPair<EInstallBundleSourceType, TSharedPtr<IInstallBundleSource>>& Pair : EnabledBundleSources)
 	{
 		Request->SourcePauseFlags.Emplace(Pair.Key, EInstallBundlePauseFlags::None);
 
@@ -1164,7 +1182,7 @@ void FDefaultInstallBundleManager::UpdateBundleSourceComplete(TSharedRef<IInstal
 		Request->CachedSourceProgress.Emplace(SourceType, MoveTemp(*Progress));
 	}
 
-	if (Request->SourceRequestResults.Num() != BundleSources.Num())
+	if (Request->SourceRequestResults.Num() != Request->RequiredSourceRequestResultsCount)
 		return;
 
 	LOG_INSTALL_BUNDLE_MAN_OVERRIDE(Request->LogVerbosityOverride, Display, TEXT("Bundle %s done waiting for all bundle sources!"), *Request->BundleName.ToString());
diff --git a/Engine/Plugins/Experimental/DefaultInstallBundleManager/Source/Public/DefaultInstallBundleManager.h b/Engine/Plugins/Experimental/DefaultInstallBundleManager/Source/Public/DefaultInstallBundleManager.h
--- a/Engine/Plugins/Experimental/DefaultInstallBundleManager/Source/Public/DefaultInstallBundleManager.h
+++ b/Engine/Plugins/Experimental/DefaultInstallBundleManager/Source/Public/DefaultInstallBundleManager.h
@@ -393,6 +393,9 @@ protected:
 
 		EInstallBundleCacheReserveResult LastCacheReserveResult = EInstallBundleCacheReserveResult::Success;
 
+		// how many results we are expected to have in the SourceRequestResults array
+		int RequiredSourceRequestResultsCount = 0;
+		// completion results from each bundle source
 		TMap<EInstallBundleSourceType, FInstallBundleSourceUpdateContentResultInfo> SourceRequestResults;
 		FText OptionalErrorText;
 		FString OptionalErrorCode;
@@ -685,6 +688,7 @@ protected:
 	void PersistentTimingStatsBegin(TSharedRef<FContentRequest> ContentRequest, InstallBundleUtil::PersistentStats::ETimingStatNames TimerStatName);
 	void PersistentTimingStatsEnd(TSharedRef<FContentRequest> ContentRequest, InstallBundleUtil::PersistentStats::ETimingStatNames TimerStatName);
 	
+	virtual TMap<EInstallBundleSourceType, TSharedPtr<IInstallBundleSource>> GetEnabledBundleSourcesForRequest(FContentRequestRef Request) const;
 	// Initialization state machine
 protected:
 	EInstallBundleManagerInitResult Init_DefaultBundleSources();
