Deprecate EInstallBundleSourceType and replace with FInstallBundleSourceType
- FInstallBundleSourceType is extensible by game code.  Possible values are automatically gathered from config
- Add a utility function to read configured bundle sources
- Fixes for enum functions in InstallBundleUtil to better support deprecated enums on clang

diff --git a/Engine/Plugins/Experimental/DefaultInstallBundleManager/Source/Private/DefaultInstallBundleManager.cpp b/Engine/Plugins/Experimental/DefaultInstallBundleManager/Source/Private/DefaultInstallBundleManager.cpp
--- a/Engine/Plugins/Experimental/DefaultInstallBundleManager/Source/Private/DefaultInstallBundleManager.cpp
+++ b/Engine/Plugins/Experimental/DefaultInstallBundleManager/Source/Private/DefaultInstallBundleManager.cpp
@@ -186,7 +186,7 @@ void LoadKeychainFromInI(FKeyChain& OutCryptoSettings)
 #endif // WITH_EDITOR
 
 FDefaultInstallBundleManager::FDefaultInstallBundleManager(FInstallBundleSourceFactoryFunction InInstallBundleSourceFactory)
-	: InstallBundleSourceFactory(InInstallBundleSourceFactory ? InInstallBundleSourceFactory : InstallBundleManagerUtil::MakeBundleSource )
+	: InstallBundleSourceFactory(InInstallBundleSourceFactory ? MoveTemp(InInstallBundleSourceFactory) : InstallBundleManagerUtil::MakeBundleSource )
 	, PersistentStats(MakeShared<InstallBundleManagerUtil::FPersistentStatContainer>())
 	, AnalyticsProvider(nullptr)
 	, StatsMap(MakeShared<InstallBundleUtil::FContentRequestStatsMap>())
@@ -476,9 +476,9 @@ void FDefaultInstallBundleManager::TickGetContentState()
 
 			FInstallBundleCombinedContentState State;
 			State.FreeSpace = 0;
-			for (const TPair<EInstallBundleSourceType, FInstallBundleCombinedContentState>& SourcePair : Request->BundleSourceContentStates)
+			for (const TPair<FInstallBundleSourceType, FInstallBundleCombinedContentState>& SourcePair : Request->BundleSourceContentStates)
 			{
-				const EInstallBundleSourceType SourceType = SourcePair.Key;
+				const FInstallBundleSourceType SourceType = SourcePair.Key;
 				const FInstallBundleCombinedContentState& SourceState = SourcePair.Value;
 
 				// combine state enums and weights
@@ -522,7 +522,7 @@ void FDefaultInstallBundleManager::TickGetContentState()
 		}
 
 		Request->Started = true;
-		for (const TPair<EInstallBundleSourceType, TSharedPtr<IInstallBundleSource>>& Pair : BundleSources)
+		for (const TPair<FInstallBundleSourceType, TSharedPtr<IInstallBundleSource>>& Pair : BundleSources)
 		{
 			Pair.Value->GetContentState(Request->BundleNames, Request->Flags, FInstallBundleGetContentStateDelegate::CreateLambda(
 				[Request, SourceType = Pair.Key](FInstallBundleCombinedContentState SourceState)
@@ -908,10 +908,10 @@ void FDefaultInstallBundleManager::TryReserveCache(FContentRequestRef Request)
 			break;
 
 		case EInstallBundleCacheReserveResult::Fail_NeedsEvict:
-			for (TPair<FName, TArray<EInstallBundleSourceType>>& ResultPair : Result.BundlesToEvict)
+			for (TPair<FName, TArray<FInstallBundleSourceType>>& ResultPair : Result.BundlesToEvict)
 			{
 				// A source can only map to one cache, so we can just append them to the list for this bundle without checking if its already there.
-				TArray<EInstallBundleSourceType>& EvictFromSources = Request->BundlesToEvictFromSourcesMap.FindOrAdd(ResultPair.Key);
+				TArray<FInstallBundleSourceType>& EvictFromSources = Request->BundlesToEvictFromSourcesMap.FindOrAdd(ResultPair.Key);
 				EvictFromSources.Append(MoveTemp(ResultPair.Value));
 			}
 			break;
@@ -990,9 +990,9 @@ void FDefaultInstallBundleManager::RequestEviction(FCacheEvictionRequestorRef Re
 {
 	LOG_INSTALL_BUNDLE_MAN_OVERRIDE(Requestor->GetLogVerbosityOverride(), Display, TEXT("Attempting to evict %i bundles for %s"), Requestor->BundlesToEvictFromSourcesMap.Num(), *Requestor->GetEvictionRequestorName());
 
-	for (const TPair<FName, TArray<EInstallBundleSourceType>>& Pair : Requestor->BundlesToEvictFromSourcesMap)
+	for (const TPair<FName, TArray<FInstallBundleSourceType>>& Pair : Requestor->BundlesToEvictFromSourcesMap)
 	{
-		for (const EInstallBundleSourceType& SourceType : Pair.Value)
+		for (FInstallBundleSourceType SourceType : Pair.Value)
 		{
 			const TSharedPtr<IInstallBundleSource>& Source = BundleSources.FindChecked(SourceType);
 			const TSharedRef<FInstallBundleCache>& BundleCache = BundleCaches.FindChecked(BundleSourceCaches.FindChecked(SourceType));
@@ -1055,7 +1055,7 @@ void FDefaultInstallBundleManager::CacheEvictionComplete(TSharedRef<IInstallBund
 
 	// Check to clear PendingEvict status
 	auto CacheEvictKey = MakeTuple(BundleCache->GetName(), InResultInfo.BundleName);
-	TArray<EInstallBundleSourceType> SourcesForCache = CachesPendingEvictToSources.FindChecked(CacheEvictKey);
+	TArray<FInstallBundleSourceType> SourcesForCache = CachesPendingEvictToSources.FindChecked(CacheEvictKey);
 	SourcesForCache.RemoveSwap(Source->GetSourceType(), EAllowShrinking::No);
 	if (SourcesForCache.Num() == 0)
 	{
@@ -1102,7 +1102,7 @@ void FDefaultInstallBundleManager::CacheEvictionComplete(TSharedRef<IInstallBund
 
 void FDefaultInstallBundleManager::CacheEvictionComplete(TSharedRef<IInstallBundleSource> Source, const FInstallBundleSourceReleaseContentResultInfo& InResultInfo, FCacheEvictionRequestorRef Requestor)
 {
-	TArray<EInstallBundleSourceType>& EvictFromSources = Requestor->BundlesToEvictFromSourcesMap.FindChecked(InResultInfo.BundleName);
+	TArray<FInstallBundleSourceType>& EvictFromSources = Requestor->BundlesToEvictFromSourcesMap.FindChecked(InResultInfo.BundleName);
 	EvictFromSources.RemoveSwap(Source->GetSourceType(), EAllowShrinking::No);
 	if (EvictFromSources.Num() == 0)
 	{
@@ -1184,7 +1184,7 @@ void FDefaultInstallBundleManager::UpdateBundleSources(FContentRequestRef Reques
 
 void FDefaultInstallBundleManager::UpdateBundleSourceComplete(TSharedRef<IInstallBundleSource> Source, FInstallBundleSourceUpdateContentResultInfo InResultInfo, FContentRequestRef Request)
 {
-	EInstallBundleSourceType SourceType = Source->GetSourceType();
+	FInstallBundleSourceType SourceType = Source->GetSourceType();
 
 	Request->SourcePauseFlags.Remove(SourceType);
 
@@ -1205,7 +1205,7 @@ void FDefaultInstallBundleManager::UpdateBundleSourceComplete(TSharedRef<IInstal
 	FBundleInfo& BundleInfo = BundleInfoMap.FindChecked(Request->BundleName);
 	BundleInfo.ContentPaths = {};
 
-	for (TPair<EInstallBundleSourceType, FInstallBundleSourceUpdateContentResultInfo>& Pair : Request->SourceRequestResults)
+	for (TPair<FInstallBundleSourceType, FInstallBundleSourceUpdateContentResultInfo>& Pair : Request->SourceRequestResults)
 	{
 		FInstallBundleSourceUpdateContentResultInfo& ResultInfo = Pair.Value;
 
@@ -1226,7 +1226,7 @@ void FDefaultInstallBundleManager::UpdateBundleSourceComplete(TSharedRef<IInstal
 					if (bCacheHit)
 					{
 						InstallBundleManagerAnalytics::FireEvent_BundleCacheHit(AnalyticsProvider.Get(),
-							BundleInfo.BundleNameString, LexToString(Pair.Key));
+							BundleInfo.BundleNameString, Pair.Key.GetNameStr());
 					}
 					else
 					{
@@ -1236,7 +1236,7 @@ void FDefaultInstallBundleManager::UpdateBundleSourceComplete(TSharedRef<IInstal
 						// 0 -> Any Data		- No Patch
 						const bool bWasPatchRequired = CacheBundleInfo->CurrentInstallSize > 0 || ResultInfo.CurrentInstallSize == 0;
 						InstallBundleManagerAnalytics::FireEvent_BundleCacheMiss(AnalyticsProvider.Get(),
-							BundleInfo.BundleNameString, LexToString(Pair.Key), bWasPatchRequired);
+							BundleInfo.BundleNameString, Pair.Key.GetNameStr(), bWasPatchRequired);
 					}
 
 					// Since the update succeeded, we know that there is no longer any install overhead
@@ -1307,7 +1307,7 @@ void FDefaultInstallBundleManager::UpdateBundleSourceComplete(TSharedRef<IInstal
 
 void FDefaultInstallBundleManager::UpdateBundleSourcePause(TSharedRef<IInstallBundleSource> Source, FInstallBundleSourcePauseInfo InPauseInfo, FContentRequestRef Request)
 {
-	EInstallBundleSourceType SourceType = Source->GetSourceType();
+	FInstallBundleSourceType SourceType = Source->GetSourceType();
 
 	Request->SourcePauseFlags.FindChecked(SourceType) = InPauseInfo.PauseFlags;
 	if (InPauseInfo.bDidPauseChange)
@@ -1330,7 +1330,7 @@ void FDefaultInstallBundleManager::UpdateBundleSources(FContentReleaseRequestRef
 		Pair.Value->Release(Request->BundleName);
 	}
 
-	for (const TPair<EInstallBundleSourceType, TSharedPtr<IInstallBundleSource>>& Pair : BundleSources)
+	for (const TPair<FInstallBundleSourceType, TSharedPtr<IInstallBundleSource>>& Pair : BundleSources)
 	{
 		if (const FName* CacheName = BundleSourceCaches.Find(Pair.Key))
 		{
@@ -1355,7 +1355,7 @@ void FDefaultInstallBundleManager::UpdateBundleSources(FContentReleaseRequestRef
 
 	Request->StepResult = EContentRequestStepResult::Waiting;
 
-	for (const TPair<EInstallBundleSourceType, TOptional<FInstallBundleSourceReleaseContentResultInfo>>& Pair : Request->SourceRemoveRequestResults)
+	for (const TPair<FInstallBundleSourceType, TOptional<FInstallBundleSourceReleaseContentResultInfo>>& Pair : Request->SourceRemoveRequestResults)
 	{
 		const TSharedPtr<IInstallBundleSource>& BundleSource = BundleSources.FindChecked(Pair.Key);
 
@@ -1368,7 +1368,7 @@ void FDefaultInstallBundleManager::UpdateBundleSources(FContentReleaseRequestRef
 		BundleSource->RequestReleaseContent(MoveTemp(Context));
 	}
 
-	for (const TPair<EInstallBundleSourceType, TOptional<FInstallBundleSourceReleaseContentResultInfo>>& Pair : Request->SourceReleaseRequestResults)
+	for (const TPair<FInstallBundleSourceType, TOptional<FInstallBundleSourceReleaseContentResultInfo>>& Pair : Request->SourceReleaseRequestResults)
 	{
 		const TSharedPtr<IInstallBundleSource>& BundleSource = BundleSources.FindChecked(Pair.Key);
 
@@ -1395,7 +1395,7 @@ void FDefaultInstallBundleManager::UpdateBundleSourceReleaseComplete(TSharedRef<
 		Result->Emplace(InResultInfo);
 	}
 
-	for (const TPair<EInstallBundleSourceType, TOptional<FInstallBundleSourceReleaseContentResultInfo>>& Pair : Request->SourceRemoveRequestResults)
+	for (const TPair<FInstallBundleSourceType, TOptional<FInstallBundleSourceReleaseContentResultInfo>>& Pair : Request->SourceRemoveRequestResults)
 	{
 		if (!Pair.Value)
 		{
@@ -1403,7 +1403,7 @@ void FDefaultInstallBundleManager::UpdateBundleSourceReleaseComplete(TSharedRef<
 		}
 	}
 
-	for (const TPair<EInstallBundleSourceType, TOptional<FInstallBundleSourceReleaseContentResultInfo>>& Pair : Request->SourceReleaseRequestResults)
+	for (const TPair<FInstallBundleSourceType, TOptional<FInstallBundleSourceReleaseContentResultInfo>>& Pair : Request->SourceReleaseRequestResults)
 	{
 		if (!Pair.Value)
 		{
@@ -1412,7 +1412,7 @@ void FDefaultInstallBundleManager::UpdateBundleSourceReleaseComplete(TSharedRef<
 	}
 
 	bool bContentWasRemoved = false;
-	for (const TPair<EInstallBundleSourceType, TOptional<FInstallBundleSourceReleaseContentResultInfo>>& Pair : Request->SourceRemoveRequestResults)
+	for (const TPair<FInstallBundleSourceType, TOptional<FInstallBundleSourceReleaseContentResultInfo>>& Pair : Request->SourceRemoveRequestResults)
 	{
 		const FInstallBundleSourceReleaseContentResultInfo& ResultInfo = *Pair.Value;
 		if (Request->Result == EInstallBundleReleaseResult::OK && ResultInfo.Result != EInstallBundleReleaseResult::OK)
@@ -1426,7 +1426,7 @@ void FDefaultInstallBundleManager::UpdateBundleSourceReleaseComplete(TSharedRef<
 		}
 	}
 
-	for (const TPair<EInstallBundleSourceType, TOptional<FInstallBundleSourceReleaseContentResultInfo>>& Pair : Request->SourceReleaseRequestResults)
+	for (const TPair<FInstallBundleSourceType, TOptional<FInstallBundleSourceReleaseContentResultInfo>>& Pair : Request->SourceReleaseRequestResults)
 	{
 		const FInstallBundleSourceReleaseContentResultInfo& ResultInfo = *Pair.Value;
 		if (Request->Result == EInstallBundleReleaseResult::OK && ResultInfo.Result != EInstallBundleReleaseResult::OK)
@@ -2401,9 +2401,9 @@ void FDefaultInstallBundleManager::TickCacheFlush()
 				continue;
 			}
 		}
-		else if (Request->SourceOrCache.HasSubtype<EInstallBundleSourceType>())
+		else if (Request->SourceOrCache.HasSubtype<FInstallBundleSourceType>())
 		{
-			if (FName* CacheName = BundleSourceCaches.Find(Request->SourceOrCache.GetSubtype<EInstallBundleSourceType>()))
+			if (FName* CacheName = BundleSourceCaches.Find(Request->SourceOrCache.GetSubtype<FInstallBundleSourceType>()))
 			{
 				pBundleCache = BundleCaches.Find(*CacheName);
 			}
@@ -2424,7 +2424,7 @@ void FDefaultInstallBundleManager::TickCacheFlush()
 
 		// Determine which sources hold the bundles we need to evict
 		bool bIsAnyPotentialEvictionPendingRelease = false;
-		TMap<FName, TArray<EInstallBundleSourceType>> BundlesToEvictFromSourcesMap;
+		TMap<FName, TArray<FInstallBundleSourceType>> BundlesToEvictFromSourcesMap;
 		if (pBundleCache)
 		{
 			const TSharedRef<FInstallBundleCache>& BundleCache = *pBundleCache;
@@ -2445,9 +2445,9 @@ void FDefaultInstallBundleManager::TickCacheFlush()
 			if (!bIsAnyPotentialEvictionPendingRelease)
 			{
 				FInstallBundleCacheFlushResult FlushResult;
-				if (Request->SourceOrCache.HasSubtype<EInstallBundleSourceType>())
+				if (Request->SourceOrCache.HasSubtype<FInstallBundleSourceType>())
 				{
-					FlushResult = BundleCache->Flush(&Request->SourceOrCache.GetSubtype<EInstallBundleSourceType>());
+					FlushResult = BundleCache->Flush(&Request->SourceOrCache.GetSubtype<FInstallBundleSourceType>());
 				}
 				else
 				{
@@ -2482,10 +2482,10 @@ void FDefaultInstallBundleManager::TickCacheFlush()
 				}
 
 				FInstallBundleCacheFlushResult FlushResult = BundleCache->Flush();
-				for (TPair<FName, TArray<EInstallBundleSourceType>>& ResultPair : FlushResult.BundlesToEvict)
+				for (TPair<FName, TArray<FInstallBundleSourceType>>& ResultPair : FlushResult.BundlesToEvict)
 				{
 					// A source can only map to one cache, so we can just append them to the list for this bundle without checking if its already there.
-					TArray<EInstallBundleSourceType>& EvictFromSources = BundlesToEvictFromSourcesMap.FindOrAdd(ResultPair.Key);
+					TArray<FInstallBundleSourceType>& EvictFromSources = BundlesToEvictFromSourcesMap.FindOrAdd(ResultPair.Key);
 					EvictFromSources.Append(MoveTemp(ResultPair.Value));
 				}
 			}
@@ -2560,7 +2560,7 @@ void FDefaultInstallBundleManager::TickPauseStatus(bool bForceCallback)
 	for (const FContentRequestRef& Request : ContentRequests[EContentRequestBatch::Install])
 	{
 		EInstallBundlePauseFlags PauseFlags = EInstallBundlePauseFlags::None;
-		for (const TPair<EInstallBundleSourceType, EInstallBundlePauseFlags>& Pair : Request->SourcePauseFlags)
+		for (const TPair<FInstallBundleSourceType, EInstallBundlePauseFlags>& Pair : Request->SourcePauseFlags)
 		{
 			PauseFlags |= Pair.Value;
 		}
@@ -2587,7 +2587,7 @@ TSet<FName> FDefaultInstallBundleManager::GetBundleDependencies(FName InBundleNa
 		*bSkippedUnknownBundles = false;
 
 		TSet<FName> SkippedUnknownBundles;
-		for (const TPair<EInstallBundleSourceType, TSharedPtr<IInstallBundleSource>>& Pair : BundleSources)
+		for (const TPair<FInstallBundleSourceType, TSharedPtr<IInstallBundleSource>>& Pair : BundleSources)
 		{
 			TSet<FName> SkippedUnknownBundlesForSource;
 			BundlesToLoad.Append(Pair.Value->GetBundleDependencies(InBundleName, &SkippedUnknownBundlesForSource));
@@ -2607,7 +2607,7 @@ TSet<FName> FDefaultInstallBundleManager::GetBundleDependencies(FName InBundleNa
 	}
 	else
 	{
-		for (const TPair<EInstallBundleSourceType, TSharedPtr<IInstallBundleSource>>& Pair : BundleSources)
+		for (const TPair<FInstallBundleSourceType, TSharedPtr<IInstallBundleSource>>& Pair : BundleSources)
 		{
 			BundlesToLoad.Append(Pair.Value->GetBundleDependencies(InBundleName));
 		}
@@ -2653,9 +2653,9 @@ TSet<FName> FDefaultInstallBundleManager::GatherBundlesForRequest(TArrayView<con
 	return GetBundleDependencies(InBundleNames);
 }
 
-EInstallBundleSourceType FDefaultInstallBundleManager::GetBundleSourceFallback(EInstallBundleSourceType Type) const
+FInstallBundleSourceType FDefaultInstallBundleManager::GetBundleSourceFallback(FInstallBundleSourceType Type) const
 {
-	const EInstallBundleSourceType* Fallback = BundleSourceFallbacks.Find(Type);
+	const FInstallBundleSourceType* Fallback = BundleSourceFallbacks.Find(Type);
 	if (Fallback)
 	{
 		return *Fallback;
@@ -2671,7 +2671,7 @@ EInstallBundleSourceUpdateBundleInfoResult FDefaultInstallBundleManager::OnUpdat
 		return EInstallBundleSourceUpdateBundleInfoResult::NotInitailized;
 	}
 
-	EInstallBundleSourceType SourceType = Source->GetSourceType();
+	FInstallBundleSourceType SourceType = Source->GetSourceType();
 
 	TSet<FName> ExistingBundles;
 	ExistingBundles.Reserve(UpdateInfo.SourceBundleInfoMap.Num());
@@ -2839,7 +2839,7 @@ EInstallBundleSourceUpdateBundleInfoResult FDefaultInstallBundleManager::OnUpdat
 
 void FDefaultInstallBundleManager::OnBundleLostRelevanceForSource(TSharedRef<IInstallBundleSource> Source, TSet<FName> BundleNames)
 {
-	EInstallBundleSourceType SourceType = Source->GetSourceType();
+	FInstallBundleSourceType SourceType = Source->GetSourceType();
 
 	for (FName BundleName : BundleNames)
 	{
@@ -2892,7 +2892,7 @@ void FDefaultInstallBundleManager::StartContentPatchCheck()
 
 	FContentPatchCheckSharedContextRef Context = MakeShared<FContentPatchCheckSharedContext>();
 
-	for (const TPair<EInstallBundleSourceType, TSharedPtr<IInstallBundleSource>>& Pair : BundleSources)
+	for (const TPair<FInstallBundleSourceType, TSharedPtr<IInstallBundleSource>>& Pair : BundleSources)
 	{
 		const TSharedPtr<IInstallBundleSource>& Source = Pair.Value;
 		Source->CheckForContentPatch(FInstallBundleSourceContentPatchResultDelegate::CreateRaw(this, &FDefaultInstallBundleManager::HandleBundleSourceContentPatchCheck, Context));
@@ -2959,7 +2959,7 @@ void FDefaultInstallBundleManager::HandleBundleSourceContentPatchCheck(TSharedRe
 void FDefaultInstallBundleManager::HandleContentPatchCheck(FContentPatchCheckSharedContextRef Context)
 {
 	EInstallBundleManagerPatchCheckResult MyResult = EInstallBundleManagerPatchCheckResult::NoPatchRequired;
-	for (const TPair<EInstallBundleSourceType, bool>& Pair : Context->Results)
+	for (const TPair<FInstallBundleSourceType, bool>& Pair : Context->Results)
 	{
 		if (Pair.Value)
 		{
@@ -2979,7 +2979,7 @@ void FDefaultInstallBundleManager::HandleContentPatchCheck(FContentPatchCheckSha
 	bIsCheckingForPatch = false;
 }
 
-bool FDefaultInstallBundleManager::HasBundleSource(EInstallBundleSourceType SourceType) const
+bool FDefaultInstallBundleManager::HasBundleSource(FInstallBundleSourceType SourceType) const
 {
 	if (InitState != EInstallBundleManagerInitState::Succeeded)
 		return false;
@@ -2987,7 +2987,7 @@ bool FDefaultInstallBundleManager::HasBundleSource(EInstallBundleSourceType Sour
 	return BundleSources.Contains(SourceType);
 }
 
-const TSharedPtr<IInstallBundleSource> FDefaultInstallBundleManager::GetBundleSource(EInstallBundleSourceType SourceType) const
+const TSharedPtr<IInstallBundleSource> FDefaultInstallBundleManager::GetBundleSource(FInstallBundleSourceType SourceType) const
 {
 	if (InitState != EInstallBundleManagerInitState::Succeeded)
 		return nullptr;
@@ -3154,7 +3154,7 @@ TValueOrError<FInstallBundleRequestInfo, EInstallBundleResult> FDefaultInstallBu
 
 		// Allow bundle sources to reject certain bundles
 		EInstallBundleSourceBundleSkipReason BundleSourceSkipReason = EInstallBundleSourceBundleSkipReason::None;
-		for (const TPair<EInstallBundleSourceType, TSharedPtr<IInstallBundleSource>>& Pair : BundleSources)
+		for (const TPair<FInstallBundleSourceType, TSharedPtr<IInstallBundleSource>>& Pair : BundleSources)
 		{
 			BundleSourceSkipReason |= Pair.Value->GetBundleSkipReason(BundleName);
 		}
@@ -3248,7 +3248,7 @@ FDelegateHandle FDefaultInstallBundleManager::GetContentState(TArrayView<const F
 	for (const FName& BundleName : AllBundles)
 	{
 		EInstallBundleSourceBundleSkipReason BundleSourceSkipReason = EInstallBundleSourceBundleSkipReason::None;
-		for (const TPair<EInstallBundleSourceType, TSharedPtr<IInstallBundleSource>>& Pair : BundleSources)
+		for (const TPair<FInstallBundleSourceType, TSharedPtr<IInstallBundleSource>>& Pair : BundleSources)
 		{
 			BundleSourceSkipReason |= Pair.Value->GetBundleSkipReason(BundleName);
 		}
@@ -3316,7 +3316,7 @@ FDelegateHandle FDefaultInstallBundleManager::GetInstallState(TArrayView<const F
 	for (const FName& BundleName : AllBundles)
 	{
 		EInstallBundleSourceBundleSkipReason BundleSourceSkipReason = EInstallBundleSourceBundleSkipReason::None;
-		for (const TPair<EInstallBundleSourceType, TSharedPtr<IInstallBundleSource>>& Pair : BundleSources)
+		for (const TPair<FInstallBundleSourceType, TSharedPtr<IInstallBundleSource>>& Pair : BundleSources)
 		{
 			BundleSourceSkipReason |= Pair.Value->GetBundleSkipReason(BundleName);
 		}
@@ -3368,7 +3368,7 @@ TValueOrError<FInstallBundleCombinedInstallState, EInstallBundleResult> FDefault
 	for (const FName& BundleName : AllBundles)
 	{
 		EInstallBundleSourceBundleSkipReason BundleSourceSkipReason = EInstallBundleSourceBundleSkipReason::None;
-		for (const TPair<EInstallBundleSourceType, TSharedPtr<IInstallBundleSource>>& Pair : BundleSources)
+		for (const TPair<FInstallBundleSourceType, TSharedPtr<IInstallBundleSource>>& Pair : BundleSources)
 		{
 			BundleSourceSkipReason |= Pair.Value->GetBundleSkipReason(BundleName);
 		}
@@ -3608,9 +3608,9 @@ TOptional<FInstallBundleCacheStats> FDefaultInstallBundleManager::GetCacheStats(
 		{
 			pBundleCache = BundleCaches.Find(SourceOrCache.GetSubtype<FName>());
 		}
-		else if (SourceOrCache.HasSubtype<EInstallBundleSourceType>())
+		else if (SourceOrCache.HasSubtype<FInstallBundleSourceType>())
 		{
-			if (FName* CacheName = BundleSourceCaches.Find(SourceOrCache.GetSubtype<EInstallBundleSourceType>()))
+			if (FName* CacheName = BundleSourceCaches.Find(SourceOrCache.GetSubtype<FInstallBundleSourceType>()))
 			{
 				pBundleCache = BundleCaches.Find(*CacheName);
 			}
@@ -3700,7 +3700,7 @@ bool FDefaultInstallBundleManager::CancelUpdateContentInternal(TArrayView<const
 				EContentRequestState State = Request->Steps[Request->iStep];
 				if (State == EContentRequestState::UpdatingBundleSources)
 				{
-					for (const TPair<EInstallBundleSourceType, TSharedPtr<IInstallBundleSource>>& Pair : BundleSources)
+					for (const TPair<FInstallBundleSourceType, TSharedPtr<IInstallBundleSource>>& Pair : BundleSources)
 					{
 						Pair.Value->CancelBundles(BundleNames);
 					}
@@ -3754,7 +3754,7 @@ bool FDefaultInstallBundleManager::CancelReleaseContentInternal(TArrayView<const
 
 void FDefaultInstallBundleManager::PauseUpdateContent(TArrayView<const FName> BundleNames)
 {
-	for (const TPair<EInstallBundleSourceType, TSharedPtr<IInstallBundleSource>>& Pair : BundleSources)
+	for (const TPair<FInstallBundleSourceType, TSharedPtr<IInstallBundleSource>>& Pair : BundleSources)
 	{
 		Pair.Value->UserPauseBundles(BundleNames);
 	}
@@ -3762,7 +3762,7 @@ void FDefaultInstallBundleManager::PauseUpdateContent(TArrayView<const FName> Bu
 
 void FDefaultInstallBundleManager::ResumeUpdateContent(TArrayView<const FName> BundleNames)
 {
-	for (const TPair<EInstallBundleSourceType, TSharedPtr<IInstallBundleSource>>& Pair : BundleSources)
+	for (const TPair<FInstallBundleSourceType, TSharedPtr<IInstallBundleSource>>& Pair : BundleSources)
 	{
 		Pair.Value->UserResumeBundles(BundleNames);
 	}
@@ -3804,7 +3804,7 @@ TOptional<FInstallBundleProgress> FDefaultInstallBundleManager::GetBundleProgres
 
 	CombinedStatus.Emplace();
 	CombinedStatus->BundleName = BundleName;
-	for (const TPair<EInstallBundleSourceType, EInstallBundlePauseFlags>& Pair : Request->SourcePauseFlags)
+	for (const TPair<FInstallBundleSourceType, EInstallBundlePauseFlags>& Pair : Request->SourcePauseFlags)
 	{
 		CombinedStatus->PauseFlags |= Pair.Value;
 	}
@@ -3827,7 +3827,7 @@ TOptional<FInstallBundleProgress> FDefaultInstallBundleManager::GetBundleProgres
 	if (State == EContentRequestState::UpdatingBundleSources && Request->Result == EInstallBundleResult::OK)
 	{
 		// Update cached status
-		for (const TPair<EInstallBundleSourceType, TSharedPtr<IInstallBundleSource>>& Pair : BundleSources)
+		for (const TPair<FInstallBundleSourceType, TSharedPtr<IInstallBundleSource>>& Pair : BundleSources)
 		{
 			TOptional<FInstallBundleSourceProgress> Progress = Pair.Value->GetBundleProgress(BundleName);
 			if (Progress)
@@ -3838,7 +3838,7 @@ TOptional<FInstallBundleProgress> FDefaultInstallBundleManager::GetBundleProgres
 
 		// Combine cached status
 		float TotalWeight = 0.0f;
-		for (const TPair<EInstallBundleSourceType, FInstallBundleSourceProgress>& Pair : Request->CachedSourceProgress)
+		for (const TPair<FInstallBundleSourceType, FInstallBundleSourceProgress>& Pair : Request->CachedSourceProgress)
 		{
 			const float SourceWeight = BundleSources.FindChecked(Pair.Key)->GetSourceWeight();
 			check(SourceWeight > 0.0f);
@@ -3899,6 +3899,9 @@ void FDefaultInstallBundleManager::StartSessionPersistentStatTracking(const FStr
 		PersistentStats->UpdateForContentState(*State, SessionName);
 	}
 
+	FInstallBundleSourceType BPSSourceType(TEXT("BuildPatchServices"));
+	check(BPSSourceType.IsValid());
+
 	for (EContentRequestBatch iBatch : TEnumRange<EContentRequestBatch>())
 	{
 		for (const FContentRequestRef& QueuedRequest : ContentRequests[iBatch])
@@ -3909,7 +3912,7 @@ void FDefaultInstallBundleManager::StartSessionPersistentStatTracking(const FStr
 				const FInstallBundleContentState* ContentState = (nullptr != State) ? State->IndividualBundleStates.Find(CurrentBundleName) : nullptr;
 				if (nullptr != ContentState)
 				{
-					FString BundleContentVersion = ContentState->Version.FindRef(EInstallBundleSourceType::BuildPatchServices);
+					FString BundleContentVersion = ContentState->Version.FindRef(BPSSourceType);
 
 					FString PreviousVersionCLString;
 					if (!BundleContentVersion.Split(TEXT("CL-"), nullptr, &PreviousVersionCLString, ESearchCase::CaseSensitive))
@@ -3917,7 +3920,7 @@ void FDefaultInstallBundleManager::StartSessionPersistentStatTracking(const FStr
 						PreviousVersionCLString = TEXT("Unknown");
 					}
 
-					FString CurrentVersion = (nullptr != State) ? State->CurrentVersion.FindRef(EInstallBundleSourceType::BuildPatchServices) : FString();
+					FString CurrentVersion = (nullptr != State) ? State->CurrentVersion.FindRef(BPSSourceType) : FString();
 					if (CurrentVersion.IsEmpty() && AnalyticsProvider)
 					{
 						//if we don't have a BPS CurrentVersion, fail back on whatever we sent to setup our Analytics Provider as that is pretty robust
@@ -4015,7 +4018,7 @@ void FDefaultInstallBundleManager::PersistentTimingStatsEnd(TSharedRef<FContentR
 EInstallBundleRequestFlags FDefaultInstallBundleManager::GetModifyableContentRequestFlags() const
 {
 	EInstallBundleRequestFlags Result = EInstallBundleRequestFlags::None;
-	for (const TPair<EInstallBundleSourceType, TSharedPtr<IInstallBundleSource>>& Pair : BundleSources)
+	for (const TPair<FInstallBundleSourceType, TSharedPtr<IInstallBundleSource>>& Pair : BundleSources)
 	{
 		Result |= Pair.Value->GetModifyableContentRequestFlags();
 	}
@@ -4025,7 +4028,7 @@ EInstallBundleRequestFlags FDefaultInstallBundleManager::GetModifyableContentReq
 
 void FDefaultInstallBundleManager::UpdateContentRequestFlags(TArrayView<const FName> BundleNames, EInstallBundleRequestFlags AddFlags, EInstallBundleRequestFlags RemoveFlags)
 {
-	for (const TPair<EInstallBundleSourceType, TSharedPtr<IInstallBundleSource>>& Pair : BundleSources)
+	for (const TPair<FInstallBundleSourceType, TSharedPtr<IInstallBundleSource>>& Pair : BundleSources)
 	{
 		EInstallBundleRequestFlags AllowedFlags = Pair.Value->GetModifyableContentRequestFlags();
 		Pair.Value->UpdateContentRequestFlags(BundleNames, AddFlags & AllowedFlags, RemoveFlags & AllowedFlags);
@@ -4079,7 +4082,7 @@ void FDefaultInstallBundleManager::SetErrorSimulationCommands(const FString& Com
 		bSimulateContentNotLatest = true;
 #endif // INSTALL_BUNDLE_ALLOW_ERROR_SIMULATION
 
-	for (const TPair<EInstallBundleSourceType, TSharedPtr<IInstallBundleSource>>& Pair : BundleSources)
+	for (const TPair<FInstallBundleSourceType, TSharedPtr<IInstallBundleSource>>& Pair : BundleSources)
 	{
 		Pair.Value->SetErrorSimulationCommands(CommandLine);
 	}
@@ -4098,59 +4101,33 @@ TSharedPtr<IAnalyticsProviderET> FDefaultInstallBundleManager::GetAnalyticsProvi
 
 EInstallBundleManagerInitResult FDefaultInstallBundleManager::Init_DefaultBundleSources()
 {
-#ifdef INSTALL_BUNDLE_SOURCES_FALLBACK_CONFIG_SECTION
-	const TCHAR* BundleSourceFallbackSection = TEXT(INSTALL_BUNDLE_SOURCES_FALLBACK_CONFIG_SECTION);
-#else
-	const TCHAR* BundleSourceFallbackSection = TEXT("InstallBundleManager.FallbackBundleSources");
-#endif // INSTALL_BUNDLE_SOURCES_FALLBACK_CONFIG_SECTION
+	TArray<FString> ConfgSources;
+	TMap<FString, FString> ConfigFallbackSources;
+	if (!InstallBundleUtil::GetConfiguredBundleSources(ConfgSources, ConfigFallbackSources))
+	{
+		return EInstallBundleManagerInitResult::ConfigurationError;
+	}
 
-	TArray<FString> ConfigBundleFallbacks;
-	if (GConfig->GetArray(BundleSourceFallbackSection, TEXT("FallbackBundleSources"), ConfigBundleFallbacks, GInstallBundleIni))
+	BundleSourceFallbacks.Empty(ConfigFallbackSources.Num());
+	for (const TPair<FString, FString>& ConfigFallback : ConfigFallbackSources)
 	{
-		BundleSourceFallbacks.Empty(ConfigBundleFallbacks.Num());
-		for (FString& ConfigFallback : ConfigBundleFallbacks)
+		FInstallBundleSourceType Key(ConfigFallback.Key);
+		FInstallBundleSourceType Value(ConfigFallback.Value);
+		if (!Key.IsValid() || !Value.IsValid())
 		{
-			// Remove parentheses
-			ConfigFallback.ReplaceInline(TEXT("("), TEXT(""));
-			ConfigFallback.ReplaceInline(TEXT(")"), TEXT(""));
-
-			TArray<FString> Tokens;
-			if (2 != ConfigFallback.ParseIntoArrayWS(Tokens, TEXT(",")))
-			{
-				ensureAlwaysMsgf(false, TEXT("Malformed entry in InstallBundleManager.FallbackBundleSources"));
-				return EInstallBundleManagerInitResult::ConfigurationError;
-			}
-
-			EInstallBundleSourceType Key = EInstallBundleSourceType::Count;
-			EInstallBundleSourceType Value = EInstallBundleSourceType::Count;
-			LexFromString(Key, *Tokens[0]);
-			LexFromString(Value, *Tokens[1]);
-			if (Key == EInstallBundleSourceType::Count || Value == EInstallBundleSourceType::Count)
-			{
-				ensureAlwaysMsgf(false, TEXT("Malformed entry in InstallBundleManager.FallbackBundleSources"));
-				return EInstallBundleManagerInitResult::ConfigurationError;
-			}
-
-			BundleSourceFallbacks.Add(Key, Value);
+			ensureAlwaysMsgf(false, TEXT("Malformed entry in InstallBundleManager.FallbackBundleSources"));
+			return EInstallBundleManagerInitResult::ConfigurationError;
 		}
-	}
-
-#ifdef INSTALL_BUNDLE_SOURCES_CONFIG_SECTION
-	const TCHAR* BundleSourceSection = TEXT(INSTALL_BUNDLE_SOURCES_CONFIG_SECTION);
-#else
-	const TCHAR* BundleSourceSection = TEXT("InstallBundleManager.BundleSources");
-#endif // INSTALL_BUNDLE_SOURCE_CONFIG_SECTION
 
-	TArray<FString> ConfigBundleSources;
-	ensureAlways(GConfig->GetArray(BundleSourceSection, TEXT("DefaultBundleSources"), ConfigBundleSources, GInstallBundleIni));
+		BundleSourceFallbacks.Add(Key, Value);
+	}
 
-	TArray<EInstallBundleSourceType> SourcesToCreate;
-	SourcesToCreate.Reserve(ConfigBundleSources.Num());
-	for (const FString& ConfigSource : ConfigBundleSources)
+	TArray<FInstallBundleSourceType> SourcesToCreate;
+	SourcesToCreate.Reserve(ConfgSources.Num());
+	for (const FString& ConfigSource : ConfgSources)
 	{
-		EInstallBundleSourceType SourceType = EInstallBundleSourceType::Count;
-		LexFromString(SourceType, *ConfigSource);
-		if (SourceType != EInstallBundleSourceType::Count)
+		FInstallBundleSourceType SourceType(ConfigSource);
+		if (SourceType.IsValid())
 		{
 			SourcesToCreate.Add(SourceType);
 		}
@@ -4178,12 +4155,12 @@ EInstallBundleManagerInitResult FDefaultInstallBundleManager::Init_DefaultBundle
 }
 
 EInstallBundleManagerInitResult FDefaultInstallBundleManager::Init_TryCreateBundleSources(
-	TArray<EInstallBundleSourceType> SourcesToCreate, TArray<TSharedPtr<IInstallBundleSource>>* OutNewSources /*= nullptr*/)
+	TArray<FInstallBundleSourceType> SourcesToCreate, TArray<TSharedPtr<IInstallBundleSource>>* OutNewSources /*= nullptr*/)
 {
 	for (int i = 0; i < SourcesToCreate.Num(); ++i)
 	{
-		EInstallBundleSourceType SourceType = SourcesToCreate[i];
-		check(SourceType != EInstallBundleSourceType::Count);
+		FInstallBundleSourceType SourceType = SourcesToCreate[i];
+		check(SourceType.IsValid());
 
 		TSharedPtr<IInstallBundleSource> Source = BundleSources.FindRef(SourceType);
 		if (!Source)
@@ -4213,7 +4190,7 @@ EInstallBundleManagerInitResult FDefaultInstallBundleManager::Init_TryCreateBund
 			}
 			else
 			{
-				EInstallBundleSourceType FallbackSourceType = FindFallbackSource(SourceType);
+				FInstallBundleSourceType FallbackSourceType = FindFallbackSource(SourceType);
 				if (FallbackSourceType == Source->GetSourceType())
 				{
 					ensureAlwaysMsgf(false, TEXT("Failed to init bundle source %s"), LexToString(SourceType));
@@ -4235,13 +4212,13 @@ EInstallBundleManagerInitResult FDefaultInstallBundleManager::Init_TryCreateBund
 	return EInstallBundleManagerInitResult::OK;
 }
 
-EInstallBundleSourceType FDefaultInstallBundleManager::FindFallbackSource(EInstallBundleSourceType SourceType)
+FInstallBundleSourceType FDefaultInstallBundleManager::FindFallbackSource(FInstallBundleSourceType SourceType)
 {
 	// See if we can find a fallback
-	EInstallBundleSourceType FallbackSourceType = SourceType;
+	FInstallBundleSourceType FallbackSourceType = SourceType;
 	for (;;)
 	{
-		EInstallBundleSourceType FallbackSourceTypeNext = GetBundleSourceFallback(FallbackSourceType);
+		FInstallBundleSourceType FallbackSourceTypeNext = GetBundleSourceFallback(FallbackSourceType);
 
 		// Check if we have reached the end of the fallback chain
 		if (FallbackSourceTypeNext == FallbackSourceType)
@@ -4299,7 +4276,7 @@ void FDefaultInstallBundleManager::AsyncInit_OnBundleSourceInitComplete(TSharedR
 	if (InitInfo.Result != EInstallBundleManagerInitResult::OK && InitInfo.bShouldUseFallbackSource)
 	{
 		// See if we can find a fallback
-		EInstallBundleSourceType FallbackSourceType = FindFallbackSource(Source->GetSourceType());
+		FInstallBundleSourceType FallbackSourceType = FindFallbackSource(Source->GetSourceType());
 		if (FallbackSourceType != Source->GetSourceType())
 		{
 			LOG_INSTALL_BUNDLE_MAN(Display, TEXT("Failed to init bundle source %s, falling back to %s"), LexToString(Source->GetSourceType()), LexToString(FallbackSourceType));
@@ -4328,7 +4305,7 @@ void FDefaultInstallBundleManager::AsyncInit_OnBundleSourceInitComplete(TSharedR
 	}
 
 	bool bHasInitializedAllSources = true;
-	for (const TPair<EInstallBundleSourceType, TOptional<FInstallBundleSourceAsyncInitInfo>>& Pair : BundleSourceInitResults)
+	for (const TPair<FInstallBundleSourceType, TOptional<FInstallBundleSourceAsyncInitInfo>>& Pair : BundleSourceInitResults)
 	{
 		if (!Pair.Value)
 		{
@@ -4340,7 +4317,7 @@ void FDefaultInstallBundleManager::AsyncInit_OnBundleSourceInitComplete(TSharedR
 	if (bHasInitializedAllSources)
 	{
 #if DO_CHECK
-		for (const TPair<EInstallBundleSourceType, TSharedPtr<IInstallBundleSource>>& Pair : BundleSources)
+		for (const TPair<FInstallBundleSourceType, TSharedPtr<IInstallBundleSource>>& Pair : BundleSources)
 		{
 			check(Pair.Value->GetInitState() != EInstallBundleManagerInitState::NotInitialized);
 		}
@@ -4358,7 +4335,7 @@ void FDefaultInstallBundleManager::AsyncInit_OnBundleSourceInitComplete(TSharedR
 		}
 		else
 		{
-			for (const TPair<EInstallBundleSourceType, TOptional<FInstallBundleSourceAsyncInitInfo>>& Pair : BundleSourceInitResults)
+			for (const TPair<FInstallBundleSourceType, TOptional<FInstallBundleSourceAsyncInitInfo>>& Pair : BundleSourceInitResults)
 			{
 				if (Pair.Value->Result != EInstallBundleManagerInitResult::OK)
 				{
@@ -4497,7 +4474,7 @@ void FDefaultInstallBundleManager::AsyncInit_InitBundleCaches()
 		}
 	}
 
-	TSet<EInstallBundleSourceType> UniqueSources;
+	TSet<FInstallBundleSourceType> UniqueSources;
 	for (FString& Mapping : ConfigBundleSourceCaches)
 	{
 		// Remove parentheses
@@ -4512,12 +4489,11 @@ void FDefaultInstallBundleManager::AsyncInit_InitBundleCaches()
 			InitResult = EInstallBundleManagerInitResult::ConfigurationError;
 		}
 
-		EInstallBundleSourceType Source = EInstallBundleSourceType::Count;
-		LexFromString(Source, *Tokens[0]);
+		FInstallBundleSourceType Source(Tokens[0]);
 		FName CacheName = *Tokens[1];
 
 		FInstallBundleCacheInitInfo* InitInfo = BundleCacheInitInfo.Find(CacheName);
-		if (Source == EInstallBundleSourceType::Count || InitInfo == nullptr)
+		if (!Source.IsValid() || InitInfo == nullptr)
 		{
 			ensureAlwaysMsgf(false, TEXT("Malformed entry in InstallBundleManager.BundleCaches"));
 			InitResult = EInstallBundleManagerInitResult::ConfigurationError;
@@ -4556,7 +4532,7 @@ void FDefaultInstallBundleManager::AsyncInit_QueryBundleInfo()
 	InitStepResult = EAsyncInitStepResult::Waiting;
 
 	BundleSourceBundleInfoQueryResults.Empty(BundleSources.Num());
-	for (const TPair<EInstallBundleSourceType, TSharedPtr<IInstallBundleSource>>& Pair : BundleSources)
+	for (const TPair<FInstallBundleSourceType, TSharedPtr<IInstallBundleSource>>& Pair : BundleSources)
 	{
 		Pair.Value->AsyncInit_QueryBundleInfo(FInstallBundleSourceQueryBundleInfoDelegate::CreateRaw(this, &FDefaultInstallBundleManager::AsyncInit_OnQueryBundleInfoComplete));
 	}
@@ -4580,7 +4556,7 @@ void FDefaultInstallBundleManager::AsyncInit_OnQueryBundleInfoComplete(TSharedRe
 	};
 
 	// Pass one, add all the bundles
-	for (const TPair<EInstallBundleSourceType, FInstallBundleSourceBundleInfoQueryResult>& SourcePair : BundleSourceBundleInfoQueryResults)
+	for (const TPair<FInstallBundleSourceType, FInstallBundleSourceBundleInfoQueryResult>& SourcePair : BundleSourceBundleInfoQueryResults)
 	{
 		for (const TPair<FName, FInstallBundleSourcePersistentBundleInfo>& BundleInfoPair : SourcePair.Value.SourceBundleInfoMap)
 		{
@@ -4633,7 +4609,7 @@ void FDefaultInstallBundleManager::AsyncInit_OnQueryBundleInfoComplete(TSharedRe
 
 		TArray<FBundleSourceRelevance> ContributingBundleSources;
 
-		for (const TPair<EInstallBundleSourceType, FInstallBundleSourceBundleInfoQueryResult>& SourcePair : BundleSourceBundleInfoQueryResults)
+		for (const TPair<FInstallBundleSourceType, FInstallBundleSourceBundleInfoQueryResult>& SourcePair : BundleSourceBundleInfoQueryResults)
 		{
 			const FInstallBundleSourcePersistentBundleInfo* SourceBundleInfo = SourcePair.Value.SourceBundleInfoMap.Find(BundleName);
 			if (SourceBundleInfo == nullptr)
@@ -4708,7 +4684,7 @@ void FDefaultInstallBundleManager::AsyncInit_SetUpdateBundleInfoCallback()
 {
 	LOG_INSTALL_BUNDLE_MAN(Display, TEXT("Setting bundle source callback to update bundle info..."));
 
-	for (const TPair<EInstallBundleSourceType, TSharedPtr<IInstallBundleSource>>& Pair : BundleSources)
+	for (const TPair<FInstallBundleSourceType, TSharedPtr<IInstallBundleSource>>& Pair : BundleSources)
 	{
 		Pair.Value->AsyncInit_SetUpdateBundleInfoCallback(
 			FInstallBundleSourceUpdateBundleInfoDelegate::CreateRaw(this, &FDefaultInstallBundleManager::OnUpdateBundleInfoFromSource),
@@ -4735,7 +4711,7 @@ void FDefaultInstallBundleManager::AsyncInit_CreateAnalyticsSession()
 		// Pick the latest content version from all sources
 		FString ContentVersion;
 		int64 MaxVersionCL = -1;
-		for (const TPair<EInstallBundleSourceType, TSharedPtr<IInstallBundleSource>>& Pair : BundleSources)
+		for (const TPair<FInstallBundleSourceType, TSharedPtr<IInstallBundleSource>>& Pair : BundleSources)
 		{
 			const TSharedPtr<IInstallBundleSource>& Source = Pair.Value;
 			FString SourceContentVersion = Source->GetContentVersion();
diff --git a/Engine/Plugins/Experimental/DefaultInstallBundleManager/Source/Private/InstallBundleManagerUtil.cpp b/Engine/Plugins/Experimental/DefaultInstallBundleManager/Source/Private/InstallBundleManagerUtil.cpp
--- a/Engine/Plugins/Experimental/DefaultInstallBundleManager/Source/Private/InstallBundleManagerUtil.cpp
+++ b/Engine/Plugins/Experimental/DefaultInstallBundleManager/Source/Private/InstallBundleManagerUtil.cpp
@@ -28,20 +28,21 @@ public:
 
 namespace InstallBundleManagerUtil
 {
-	TSharedPtr<IInstallBundleSource> MakeBundleSource(EInstallBundleSourceType Type)
+	TSharedPtr<IInstallBundleSource> MakeBundleSource(FInstallBundleSourceType Type)
 	{
-		switch (Type)
+		if (Type.GetName() == TEXTVIEW("Bulk"))
 		{
-		case EInstallBundleSourceType::Bulk:
 			return MakeShared<FInstallBundleSourceBulk>();
-			
+		}
+
 #if WITH_PLATFORM_INSTALL_BUNDLE_SOURCE
-		case EInstallBundleSourceType::Platform:
+		if (Type.GetName() == TEXTVIEW("Platform"))
+		{
 			return MakePlatformBundleSource();
-#endif
 		}
+#endif // WITH_PLATFORM_INSTALL_BUNDLE_SOURCE
 
-		LOG_INSTALL_BUNDLE_MAN(Fatal, TEXT("Can't make EInstallBundleSourceType %s"), LexToString(Type));
+		LOG_INSTALL_BUNDLE_MAN(Fatal, TEXT("Can't make InstallBundleSourceType %.*s"), Type.GetName().Len(), Type.GetName().GetData());
 		return nullptr;
 	}
 
@@ -67,7 +68,7 @@ namespace InstallBundleManagerUtil
 		return PinnedJournalThreadPool;
 	}
 
-	bool LoadBundleSourceBundleInfoFromConfig(EInstallBundleSourceType SourceType, const FConfigFile& InstallBundleConfig, const FString& Section, FInstallBundleSourcePersistentBundleInfo& OutInfo)
+	bool LoadBundleSourceBundleInfoFromConfig(FInstallBundleSourceType SourceType, const FConfigFile& InstallBundleConfig, const FString& Section, FInstallBundleSourcePersistentBundleInfo& OutInfo)
 	{
 		if (!Section.StartsWith(InstallBundleUtil::GetInstallBundleSectionPrefix()))
 			return false;
@@ -78,7 +79,7 @@ namespace InstallBundleManagerUtil
 
 		TArray<FString> ExcludedBundleSources;
 		InstallBundleConfig.GetArray(*Section, TEXT("ExcludedBundleSources"), ExcludedBundleSources);
-		if (ExcludedBundleSources.Contains(LexToString(SourceType)))
+		if (ExcludedBundleSources.Contains(SourceType.GetName()))
 			return false;
 
 		if (!InstallBundleConfig.GetBool(*Section, TEXT("IsStartup"), OutInfo.bIsStartup))
@@ -117,7 +118,7 @@ namespace InstallBundleManagerUtil
 		TArray<FString> CachedBySources;
 		if (InstallBundleConfig.GetArray(*Section, TEXT("CachedBySource"), CachedBySources))
 		{
-			if (CachedBySources.Contains(LexToString(SourceType)))
+			if (CachedBySources.Contains(SourceType.GetName()))
 			{
 				OutInfo.bIsCached = true;
 			}
@@ -299,7 +300,7 @@ namespace InstallBundleManagerUtil
 		FoundData.ResetShouldSendBGAnalytics();
 	}
 
-	void FPersistentStatContainer::UpdateForBundleSource(const FInstallBundleSourceUpdateContentResultInfo& BundleSourceResult, EInstallBundleSourceType SourceType, const FString& BundleName)
+	void FPersistentStatContainer::UpdateForBundleSource(const FInstallBundleSourceUpdateContentResultInfo& BundleSourceResult, FInstallBundleSourceType SourceType, const FString& BundleName)
 	{
 		FBundleAnalyticsData& BundleAnalyticsData = BundleAnalyticsDataMap.FindOrAdd(BundleName);
 		
@@ -370,7 +371,7 @@ namespace InstallBundleManagerUtil
 		FPersistentStatsInformation NewStatsInfo;
 		NewStatsInfo.SessionName = SessionName;
 		
-		TSet<EInstallBundleSourceType> SourcesThatDidWork;
+		TSet<FInstallBundleSourceType> SourcesThatDidWork;
 		
 		InstallBundleUtil::PersistentStats::FSessionPersistentStats* SessionStats = SessionPersistentStatMap.Find(SessionName);
 		if (nullptr != SessionStats)
@@ -546,16 +547,16 @@ namespace InstallBundleManagerUtil
 		}
 		
 		//Create BundleSourcesThatDidWork string from SourcesThatDidWork map
-		for (EInstallBundleSourceType BundleSourceType : SourcesThatDidWork)
+		for (FInstallBundleSourceType BundleSourceType : SourcesThatDidWork)
 		{
 			if (NewStatsInfo.BundleSourcesThatDidWork.IsEmpty())
 			{
-				NewStatsInfo.BundleSourcesThatDidWork = LexToString(BundleSourceType);
+				NewStatsInfo.BundleSourcesThatDidWork = BundleSourceType.GetName();
 			}
 			else
 			{
 				NewStatsInfo.BundleSourcesThatDidWork.Append(TEXT(", "));
-				NewStatsInfo.BundleSourcesThatDidWork.Append(LexToString(BundleSourceType));
+				NewStatsInfo.BundleSourcesThatDidWork.Append(BundleSourceType.GetName());
 			}
 		}
 		
diff --git a/Engine/Plugins/Experimental/DefaultInstallBundleManager/Source/Private/InstallBundleSourceBulk.cpp b/Engine/Plugins/Experimental/DefaultInstallBundleManager/Source/Private/InstallBundleSourceBulk.cpp
--- a/Engine/Plugins/Experimental/DefaultInstallBundleManager/Source/Private/InstallBundleSourceBulk.cpp
+++ b/Engine/Plugins/Experimental/DefaultInstallBundleManager/Source/Private/InstallBundleSourceBulk.cpp
@@ -451,6 +451,11 @@ EInstallBundleInstallState FInstallBundleSourceBulk::GetBundleInstallState(FName
 	return EInstallBundleInstallState::UpToDate;
 }
 
+FInstallBundleSourceType FInstallBundleSourceBulk::GetSourceType() const
+{ 
+	return FInstallBundleSourceType(TEXT("Bulk")); 
+}
+
 FInstallBundleSourceInitInfo FInstallBundleSourceBulk::Init(
 	TSharedRef<InstallBundleUtil::FContentRequestStatsMap> InRequestStats,
 	TSharedPtr<IAnalyticsProviderET> InAnalyticsProvider,
diff --git a/Engine/Plugins/Experimental/DefaultInstallBundleManager/Source/Private/InstallBundleSourcePlatformBase.cpp b/Engine/Plugins/Experimental/DefaultInstallBundleManager/Source/Private/InstallBundleSourcePlatformBase.cpp
--- a/Engine/Plugins/Experimental/DefaultInstallBundleManager/Source/Private/InstallBundleSourcePlatformBase.cpp
+++ b/Engine/Plugins/Experimental/DefaultInstallBundleManager/Source/Private/InstallBundleSourcePlatformBase.cpp
@@ -24,9 +24,9 @@ FInstallBundleSourcePlatformBase::~FInstallBundleSourcePlatformBase()
 }
 
 
-EInstallBundleSourceType FInstallBundleSourcePlatformBase::GetSourceType() const
+FInstallBundleSourceType FInstallBundleSourcePlatformBase::GetSourceType() const
 {
-	return EInstallBundleSourceType::Platform;
+	return FInstallBundleSourceType(TEXT("Platform"));
 }
 
 FInstallBundleSourceInitInfo FInstallBundleSourcePlatformBase::Init(
diff --git a/Engine/Plugins/Experimental/DefaultInstallBundleManager/Source/Public/DefaultInstallBundleManager.h b/Engine/Plugins/Experimental/DefaultInstallBundleManager/Source/Public/DefaultInstallBundleManager.h
--- a/Engine/Plugins/Experimental/DefaultInstallBundleManager/Source/Public/DefaultInstallBundleManager.h
+++ b/Engine/Plugins/Experimental/DefaultInstallBundleManager/Source/Public/DefaultInstallBundleManager.h
@@ -137,7 +137,7 @@ protected:
 
 	struct FBundleSourceRelevance
 	{
-		EInstallBundleSourceType SourceType;
+		FInstallBundleSourceType SourceType;
 		bool bIsRelevant = true;
 
 		bool operator==(const FBundleSourceRelevance& Other) const { return SourceType == Other.SourceType; }
@@ -211,7 +211,7 @@ protected:
 
 	struct FGetContentStateRequest
 	{
-		TMap<EInstallBundleSourceType, FInstallBundleCombinedContentState> BundleSourceContentStates;
+		TMap<FInstallBundleSourceType, FInstallBundleCombinedContentState> BundleSourceContentStates;
 
 		TArray<FName> BundleNames;
 
@@ -322,7 +322,7 @@ protected:
 
 	struct FCacheEvictionRequestor
 	{
-		TMap<FName, TArray<EInstallBundleSourceType>> BundlesToEvictFromSourcesMap;
+		TMap<FName, TArray<FInstallBundleSourceType>> BundlesToEvictFromSourcesMap;
 
 		virtual ~FCacheEvictionRequestor() {}
 
@@ -347,9 +347,9 @@ protected:
 			{
 				return FString::Printf(TEXT("CacheFlush(%s)"), *SourceOrCache.GetSubtype<FName>().ToString());
 			}
-			else if (SourceOrCache.HasSubtype<EInstallBundleSourceType>())
+			else if (SourceOrCache.HasSubtype<FInstallBundleSourceType>())
 			{
-				return FString::Printf(TEXT("CacheFlush(%s)"), LexToString(SourceOrCache.GetSubtype<EInstallBundleSourceType>()));
+				return FString::Printf(TEXT("CacheFlush(%s)"), LexToString(SourceOrCache.GetSubtype<FInstallBundleSourceType>()));
 			}
 			else
 			{
@@ -387,7 +387,7 @@ protected:
 
 		FName BundleName;
 
-		TMap<EInstallBundleSourceType, EInstallBundlePauseFlags> SourcePauseFlags;
+		TMap<FInstallBundleSourceType, EInstallBundlePauseFlags> SourcePauseFlags;
 		EInstallBundlePauseFlags LastSentPauseFlags = EInstallBundlePauseFlags::None;
 		bool bForcePauseCallback = false;
 
@@ -396,11 +396,11 @@ protected:
 		// how many results we are expected to have in the SourceRequestResults array
 		int RequiredSourceRequestResultsCount = 0;
 		// completion results from each bundle source
-		TMap<EInstallBundleSourceType, FInstallBundleSourceUpdateContentResultInfo> SourceRequestResults;
+		TMap<FInstallBundleSourceType, FInstallBundleSourceUpdateContentResultInfo> SourceRequestResults;
 		FText OptionalErrorText;
 		FString OptionalErrorCode;
 
-		TMap<EInstallBundleSourceType, FInstallBundleSourceProgress> CachedSourceProgress;
+		TMap<FInstallBundleSourceType, FInstallBundleSourceProgress> CachedSourceProgress;
 
 		InstallBundleUtil::FContentRequestSharedContextPtr RequestSharedContext;
 
@@ -456,8 +456,8 @@ protected:
 		
 		FName BundleName;
 
-		TMap<EInstallBundleSourceType, TOptional<FInstallBundleSourceReleaseContentResultInfo>> SourceReleaseRequestResults;
-		TMap<EInstallBundleSourceType, TOptional<FInstallBundleSourceReleaseContentResultInfo>> SourceRemoveRequestResults;
+		TMap<FInstallBundleSourceType, TOptional<FInstallBundleSourceReleaseContentResultInfo>> SourceReleaseRequestResults;
+		TMap<FInstallBundleSourceType, TOptional<FInstallBundleSourceReleaseContentResultInfo>> SourceRemoveRequestResults;
 
 		ELogVerbosity::Type LogVerbosityOverride = ELogVerbosity::NoLogging;
 
@@ -470,16 +470,16 @@ protected:
 
 	struct FContentPatchCheckSharedContext
 	{
-		TMap<EInstallBundleSourceType, bool> Results;
+		TMap<FInstallBundleSourceType, bool> Results;
 	};
 	using FContentPatchCheckSharedContextRef = TSharedRef<FContentPatchCheckSharedContext>;
 
 public:
-	typedef TFunction<TSharedPtr<IInstallBundleSource>(EInstallBundleSourceType)> FInstallBundleSourceFactoryFunction;
+	typedef TUniqueFunction<TSharedPtr<IInstallBundleSource>(FInstallBundleSourceType)> FInstallBundleSourceFactoryFunction;
 
 	UE_DEPRECATED(5.4, "GInstallBundleManagerIni is deprecated, use InstallBundle.ini hierarchy instead.")
 	FDefaultInstallBundleManager(const TCHAR* InConfigBaseName, FInstallBundleSourceFactoryFunction InBundleSourceFactory = nullptr)
-		: FDefaultInstallBundleManager(InBundleSourceFactory)
+		: FDefaultInstallBundleManager(MoveTemp(InBundleSourceFactory))
 	{}
 
 	FDefaultInstallBundleManager(FInstallBundleSourceFactoryFunction InBundleSourceFactory = nullptr);
@@ -596,7 +596,7 @@ protected:
 
 	TSet<FName> GatherBundlesForRequest(TArrayView<const FName> InBundleNames);
 
-	EInstallBundleSourceType GetBundleSourceFallback(EInstallBundleSourceType Type) const;
+	FInstallBundleSourceType GetBundleSourceFallback(FInstallBundleSourceType Type) const;
 
 	EInstallBundleSourceUpdateBundleInfoResult OnUpdateBundleInfoFromSource(TSharedRef<IInstallBundleSource> Source, FInstallBundleSourceUpdateBundleInfoResult Result);
 	void OnBundleLostRelevanceForSource(TSharedRef<IInstallBundleSource> Source, TSet<FName> BundleNames);
@@ -612,9 +612,9 @@ protected:
 
 	// IInstallBundleManager interface
 public:
-	virtual bool HasBundleSource(EInstallBundleSourceType SourceType) const override;
+	virtual bool HasBundleSource(FInstallBundleSourceType SourceType) const override;
 
-	virtual const TSharedPtr<IInstallBundleSource> GetBundleSource(EInstallBundleSourceType SourceType) const override;
+	virtual const TSharedPtr<IInstallBundleSource> GetBundleSource(FInstallBundleSourceType SourceType) const override;
 
 
 	virtual FDelegateHandle PushInitErrorCallback(FInstallBundleManagerInitErrorHandler Callback) override;
@@ -694,9 +694,9 @@ protected:
 	// Initialization state machine
 protected:
 	EInstallBundleManagerInitResult Init_DefaultBundleSources();
-	EInstallBundleManagerInitResult Init_TryCreateBundleSources(TArray<EInstallBundleSourceType> SourcesToCreate, TArray<TSharedPtr<IInstallBundleSource>>* OutNewSources = nullptr);
+	EInstallBundleManagerInitResult Init_TryCreateBundleSources(TArray<FInstallBundleSourceType> SourcesToCreate, TArray<TSharedPtr<IInstallBundleSource>>* OutNewSources = nullptr);
 
-	EInstallBundleSourceType FindFallbackSource(EInstallBundleSourceType SourceType);
+	FInstallBundleSourceType FindFallbackSource(FInstallBundleSourceType SourceType);
 	void AsyncInit_InitBundleSources();
 	void AsyncInit_OnBundleSourceInitComplete(TSharedRef<IInstallBundleSource> Source, FInstallBundleSourceAsyncInitInfo InInitInfo);
 	void AsyncInit_InitBundleCaches();
@@ -721,19 +721,19 @@ protected:
 	TMap<FName, FBundleInfo> BundleInfoMap;
 	TSet<FName> BundlesInfosToPrune;
 
-	TMap<EInstallBundleSourceType, TSharedPtr<IInstallBundleSource>> BundleSources;
-	TMap<EInstallBundleSourceType, EInstallBundleSourceType> BundleSourceFallbacks;
+	TMap<FInstallBundleSourceType, TSharedPtr<IInstallBundleSource>> BundleSources;
+	TMap<FInstallBundleSourceType, FInstallBundleSourceType> BundleSourceFallbacks;
 
 	TMap<FName, TSharedRef<FInstallBundleCache>> BundleCaches;
-	TMap<EInstallBundleSourceType, FName> BundleSourceCaches;
+	TMap<FInstallBundleSourceType, FName> BundleSourceCaches;
 	TMap<FName, uint64> BundleCacheSizeOverrides;
 
-	TMap<TTuple<EInstallBundleSourceType, FName>, TArray<FCacheEvictionRequestorRef>> PendingCacheEvictions; // (Source, Bundle) -> List of requestors
-	TMap<TTuple<FName, FName>, TArray<EInstallBundleSourceType>> CachesPendingEvictToSources; // (Cache, Bundle) -> List of Sources
+	TMap<TTuple<FInstallBundleSourceType, FName>, TArray<FCacheEvictionRequestorRef>> PendingCacheEvictions; // (Source, Bundle) -> List of requestors
+	TMap<TTuple<FName, FName>, TArray<FInstallBundleSourceType>> CachesPendingEvictToSources; // (Cache, Bundle) -> List of Sources
 
 	// Only used during Init
-	TMap<EInstallBundleSourceType, TOptional<FInstallBundleSourceAsyncInitInfo>> BundleSourceInitResults;
-	TMap<EInstallBundleSourceType, FInstallBundleSourceBundleInfoQueryResult> BundleSourceBundleInfoQueryResults;
+	TMap<FInstallBundleSourceType, TOptional<FInstallBundleSourceAsyncInitInfo>> BundleSourceInitResults;
+	TMap<FInstallBundleSourceType, FInstallBundleSourceBundleInfoQueryResult> BundleSourceBundleInfoQueryResults;
 
 	// Init
 	EInstallBundleManagerInitState InitState = EInstallBundleManagerInitState::NotInitialized;
diff --git a/Engine/Plugins/Experimental/DefaultInstallBundleManager/Source/Public/InstallBundleManagerUtil.h b/Engine/Plugins/Experimental/DefaultInstallBundleManager/Source/Public/InstallBundleManagerUtil.h
--- a/Engine/Plugins/Experimental/DefaultInstallBundleManager/Source/Public/InstallBundleManagerUtil.h
+++ b/Engine/Plugins/Experimental/DefaultInstallBundleManager/Source/Public/InstallBundleManagerUtil.h
@@ -18,7 +18,7 @@ DECLARE_LOG_CATEGORY_EXTERN(LogDefaultInstallBundleManager, Display, All);
 
 namespace InstallBundleManagerUtil
 {
-	DEFAULTINSTALLBUNDLEMANAGER_API TSharedPtr<IInstallBundleSource> MakeBundleSource(EInstallBundleSourceType Type);
+	DEFAULTINSTALLBUNDLEMANAGER_API TSharedPtr<IInstallBundleSource> MakeBundleSource(FInstallBundleSourceType Type);
 
 #if WITH_PLATFORM_INSTALL_BUNDLE_SOURCE
 	DEFAULTINSTALLBUNDLEMANAGER_API TSharedPtr<IInstallBundleSource> MakePlatformBundleSource();
@@ -30,7 +30,7 @@ namespace InstallBundleManagerUtil
 
 	// Fills out a FInstallBundleSourceBundleInfo from the specified config section
 	// Returns false if the provided config section is not a bundle definition section.
-	DEFAULTINSTALLBUNDLEMANAGER_API bool LoadBundleSourceBundleInfoFromConfig(EInstallBundleSourceType SourceType, const FConfigFile& InstallBundleConfig, const FString& Section, FInstallBundleSourcePersistentBundleInfo& OutInfo);
+	DEFAULTINSTALLBUNDLEMANAGER_API bool LoadBundleSourceBundleInfoFromConfig(FInstallBundleSourceType SourceType, const FConfigFile& InstallBundleConfig, const FString& Section, FInstallBundleSourcePersistentBundleInfo& OutInfo);
 
 	// Traverses bundle config sections and loads all dependencies for InBundleName, including InBundleName
 	// Sets bSkippedUnknownBundles if a config section for InBundleName or a dependency can't be found.
@@ -71,7 +71,7 @@ namespace InstallBundleManagerUtil
 		void AddRequiredBundlesForSession(const FString& SessionName, const TArray<FName>& BundleNames);
 
 		void UpdateForContentState(const FInstallBundleCombinedContentState& ContentState, const FString& SessionName);
-		void UpdateForBundleSource(const FInstallBundleSourceUpdateContentResultInfo& BundleSourceResult, EInstallBundleSourceType SourceType, const FString& BundleName);
+		void UpdateForBundleSource(const FInstallBundleSourceUpdateContentResultInfo& BundleSourceResult, FInstallBundleSourceType SourceType, const FString& BundleName);
 		
 	private:
 		void SendEnteringBackgroundAnalytic();
@@ -166,7 +166,7 @@ namespace InstallBundleManagerUtil
 		{
 		public:
 			//Store if each bundle source type did any work for this bundle. If its in this map it did work.
-			TSet<EInstallBundleSourceType> BundleSourcesThatDidWorkMap;
+			TSet<FInstallBundleSourceType> BundleSourcesThatDidWorkMap;
 			
 			FBundleAnalyticsData()
 			: BundleSourcesThatDidWorkMap()
diff --git a/Engine/Plugins/Experimental/DefaultInstallBundleManager/Source/Public/InstallBundleSourceBulk.h b/Engine/Plugins/Experimental/DefaultInstallBundleManager/Source/Public/InstallBundleSourceBulk.h
--- a/Engine/Plugins/Experimental/DefaultInstallBundleManager/Source/Public/InstallBundleSourceBulk.h
+++ b/Engine/Plugins/Experimental/DefaultInstallBundleManager/Source/Public/InstallBundleSourceBulk.h
@@ -27,7 +27,7 @@ protected:
 
 	// IInstallBundleSource Interface
 public:
-	virtual EInstallBundleSourceType GetSourceType() const override { return EInstallBundleSourceType::Bulk; }
+	virtual FInstallBundleSourceType GetSourceType() const override;
 	virtual float GetSourceWeight() const override { return 0.1f; }  // Low weight since all this source does in mount
 
 	virtual FInstallBundleSourceInitInfo Init(
diff --git a/Engine/Plugins/Experimental/DefaultInstallBundleManager/Source/Public/InstallBundleSourcePlatformBase.h b/Engine/Plugins/Experimental/DefaultInstallBundleManager/Source/Public/InstallBundleSourcePlatformBase.h
--- a/Engine/Plugins/Experimental/DefaultInstallBundleManager/Source/Public/InstallBundleSourcePlatformBase.h
+++ b/Engine/Plugins/Experimental/DefaultInstallBundleManager/Source/Public/InstallBundleSourcePlatformBase.h
@@ -21,7 +21,7 @@ public:
 
 	// IInstallBundleSource Interface
 public:
-	virtual EInstallBundleSourceType GetSourceType() const override;
+	virtual FInstallBundleSourceType GetSourceType() const override;
 
 	virtual FInstallBundleSourceInitInfo Init(
 		TSharedRef<InstallBundleUtil::FContentRequestStatsMap> InRequestStats,
diff --git a/Engine/Source/Runtime/InstallBundleManager/Private/InstallBundleCache.cpp b/Engine/Source/Runtime/InstallBundleManager/Private/InstallBundleCache.cpp
--- a/Engine/Source/Runtime/InstallBundleManager/Private/InstallBundleCache.cpp
+++ b/Engine/Source/Runtime/InstallBundleManager/Private/InstallBundleCache.cpp
@@ -16,7 +16,7 @@ void FInstallBundleCache::Init(FInstallBundleCacheInitInfo InitInfo)
 	TotalSize = InitInfo.Size;
 }
 
-void FInstallBundleCache::AddOrUpdateBundle(EInstallBundleSourceType Source, const FInstallBundleCacheBundleInfo& AddInfo)
+void FInstallBundleCache::AddOrUpdateBundle(FInstallBundleSourceType Source, const FInstallBundleCacheBundleInfo& AddInfo)
 {
 	CSV_SCOPED_TIMING_STAT(InstallBundleManager, FInstallBundleCache_AddOrUpdateBundle);
 
@@ -32,11 +32,11 @@ void FInstallBundleCache::AddOrUpdateBundle(EInstallBundleSourceType Source, con
 	CheckInvariants();
 }
 
-void FInstallBundleCache::RemoveBundle(EInstallBundleSourceType Source, FName BundleName)
+void FInstallBundleCache::RemoveBundle(FInstallBundleSourceType Source, FName BundleName)
 {
 	CSV_SCOPED_TIMING_STAT(InstallBundleManager, FInstallBundleCache_RemoveBundle);
 
-	TMap<EInstallBundleSourceType, FPerSourceBundleCacheInfo>* SourcesMap = PerSourceCacheInfo.Find(BundleName);
+	TMap<FInstallBundleSourceType, FPerSourceBundleCacheInfo>* SourcesMap = PerSourceCacheInfo.Find(BundleName);
 	if (SourcesMap)
 	{
 		SourcesMap->Remove(Source);
@@ -65,11 +65,11 @@ TOptional<FInstallBundleCacheBundleInfo> FInstallBundleCache::GetBundleInfo(FNam
 	return Ret;
 }
 
-TOptional<FInstallBundleCacheBundleInfo> FInstallBundleCache::GetBundleInfo(EInstallBundleSourceType Source, FName BundleName) const
+TOptional<FInstallBundleCacheBundleInfo> FInstallBundleCache::GetBundleInfo(FInstallBundleSourceType Source, FName BundleName) const
 {
 	TOptional<FInstallBundleCacheBundleInfo> Ret;
 
-	const TMap<EInstallBundleSourceType, FPerSourceBundleCacheInfo>* SourcesMap = PerSourceCacheInfo.Find(BundleName);
+	const TMap<FInstallBundleSourceType, FPerSourceBundleCacheInfo>* SourcesMap = PerSourceCacheInfo.Find(BundleName);
 	if (SourcesMap)
 	{
 		const FPerSourceBundleCacheInfo* SourceInfo = SourcesMap->Find(Source);
@@ -181,7 +181,7 @@ FInstallBundleCacheReserveResult FInstallBundleCache::Reserve(FName BundleName)
 			if (TotalSize < UsedSize - CanFreeSpace + SizeNeeded)
 			{
 				CanFreeSpace += BundleSize;
-				TArray<EInstallBundleSourceType>& SourcesToEvictFrom = Result.BundlesToEvict.Add(Pair.Key);
+				TArray<FInstallBundleSourceType>& SourcesToEvictFrom = Result.BundlesToEvict.Add(Pair.Key);
 				PerSourceCacheInfo.FindChecked(Pair.Key).GenerateKeyArray(SourcesToEvictFrom);
 			}
 		}
@@ -190,7 +190,7 @@ FInstallBundleCacheReserveResult FInstallBundleCache::Reserve(FName BundleName)
 			// Bundle manager must wait for all previous pending evictions to complete
 			// to ensure that there is actually enough free space in the cache
 			// before installing a bundle
-			TArray<EInstallBundleSourceType>& SourcesToEvictFrom = Result.BundlesToEvict.Add(Pair.Key);
+			TArray<FInstallBundleSourceType>& SourcesToEvictFrom = Result.BundlesToEvict.Add(Pair.Key);
 			PerSourceCacheInfo.FindChecked(Pair.Key).GenerateKeyArray(SourcesToEvictFrom);
 		}
 	}
@@ -212,7 +212,7 @@ FInstallBundleCacheReserveResult FInstallBundleCache::Reserve(FName BundleName)
 	return Result;
 }
 
-FInstallBundleCacheFlushResult FInstallBundleCache::Flush(EInstallBundleSourceType* Source /*= nullptr*/)
+FInstallBundleCacheFlushResult FInstallBundleCache::Flush(FInstallBundleSourceType* Source /*= nullptr*/)
 {
 	CSV_SCOPED_TIMING_STAT(InstallBundleManager, FInstallBundleCache_Flush);
 
@@ -230,13 +230,13 @@ FInstallBundleCacheFlushResult FInstallBundleCache::Flush(EInstallBundleSourceTy
 		{
 			if (PerSourceCacheInfo.FindChecked(Pair.Key).Contains(*Source))
 			{
-				TArray<EInstallBundleSourceType>& SourcesToEvictFrom = Result.BundlesToEvict.Add(Pair.Key);
+				TArray<FInstallBundleSourceType>& SourcesToEvictFrom = Result.BundlesToEvict.Add(Pair.Key);
 				SourcesToEvictFrom.Add(*Source);
 			}
 		}
 		else
 		{
-			TArray<EInstallBundleSourceType>& SourcesToEvictFrom = Result.BundlesToEvict.Add(Pair.Key);
+			TArray<FInstallBundleSourceType>& SourcesToEvictFrom = Result.BundlesToEvict.Add(Pair.Key);
 			PerSourceCacheInfo.FindChecked(Pair.Key).GenerateKeyArray(SourcesToEvictFrom);
 		}
 	}
@@ -253,9 +253,9 @@ bool FInstallBundleCache::Contains(FName BundleName) const
 	return CacheInfo.Contains(BundleName);
 }
 
-bool FInstallBundleCache::Contains(EInstallBundleSourceType Source, FName BundleName) const
+bool FInstallBundleCache::Contains(FInstallBundleSourceType Source, FName BundleName) const
 {
-	const TMap<EInstallBundleSourceType, FPerSourceBundleCacheInfo>* SourcesMap = PerSourceCacheInfo.Find(BundleName);
+	const TMap<FInstallBundleSourceType, FPerSourceBundleCacheInfo>* SourcesMap = PerSourceCacheInfo.Find(BundleName);
 	if (SourcesMap)
 	{
 		return SourcesMap->Contains(Source);
@@ -374,12 +374,12 @@ void FInstallBundleCache::CheckInvariants() const
 
 	for (const TPair<FName, FBundleCacheInfo>& CachePair : CacheInfo)
 	{
-		const TMap<EInstallBundleSourceType, FPerSourceBundleCacheInfo>* SourcesMap = PerSourceCacheInfo.Find(CachePair.Key);
+		const TMap<FInstallBundleSourceType, FPerSourceBundleCacheInfo>* SourcesMap = PerSourceCacheInfo.Find(CachePair.Key);
 		check(SourcesMap);
 
 		uint64 FullInstallSize = 0;
 		uint64 CurrentInstallSize = 0;
-		for (const TPair<EInstallBundleSourceType, FPerSourceBundleCacheInfo>& Pair : *SourcesMap)
+		for (const TPair<FInstallBundleSourceType, FPerSourceBundleCacheInfo>& Pair : *SourcesMap)
 		{
 			FullInstallSize += Pair.Value.FullInstallSize;
 			CurrentInstallSize += Pair.Value.CurrentInstallSize;
@@ -545,7 +545,7 @@ void FInstallBundleCache::UpdateCacheInfoFromSourceInfo(FName BundleName)
 {
 	CSV_SCOPED_TIMING_STAT(InstallBundleManager, FInstallBundleCache_UpdateCacheInfoFromSourceInfo);
 
-	TMap<EInstallBundleSourceType, FPerSourceBundleCacheInfo>* SourcesMap = PerSourceCacheInfo.Find(BundleName);
+	TMap<FInstallBundleSourceType, FPerSourceBundleCacheInfo>* SourcesMap = PerSourceCacheInfo.Find(BundleName);
 	if (SourcesMap == nullptr)
 	{
 		CacheInfo.Remove(BundleName);
@@ -564,7 +564,7 @@ void FInstallBundleCache::UpdateCacheInfoFromSourceInfo(FName BundleName)
 	uint64 FullInstallSize = 0;
 	uint64 InstallOverheadSize = 0;
 	uint64 CurrentInstallSize = 0;
-	for (const TPair<EInstallBundleSourceType, FPerSourceBundleCacheInfo>& Pair : *SourcesMap)
+	for (const TPair<FInstallBundleSourceType, FPerSourceBundleCacheInfo>& Pair : *SourcesMap)
 	{
 		FullInstallSize += Pair.Value.FullInstallSize;
 		InstallOverheadSize += Pair.Value.InstallOverheadSize;
diff --git a/Engine/Source/Runtime/InstallBundleManager/Private/InstallBundleManagerInterface.cpp b/Engine/Source/Runtime/InstallBundleManager/Private/InstallBundleManagerInterface.cpp
--- a/Engine/Source/Runtime/InstallBundleManager/Private/InstallBundleManagerInterface.cpp
+++ b/Engine/Source/Runtime/InstallBundleManager/Private/InstallBundleManagerInterface.cpp
@@ -48,7 +48,7 @@ TSharedPtr<IInstallBundleManager> IInstallBundleManager::GetPlatformInstallBundl
 	return {};
 }
 
-const TSharedPtr<IInstallBundleSource> IInstallBundleManager::GetBundleSource(EInstallBundleSourceType SourceType) const
+const TSharedPtr<IInstallBundleSource> IInstallBundleManager::GetBundleSource(FInstallBundleSourceType SourceType) const
 {
 	return {};
 }
diff --git a/Engine/Source/Runtime/InstallBundleManager/Private/InstallBundleTypes.cpp b/Engine/Source/Runtime/InstallBundleManager/Private/InstallBundleTypes.cpp
--- a/Engine/Source/Runtime/InstallBundleManager/Private/InstallBundleTypes.cpp
+++ b/Engine/Source/Runtime/InstallBundleManager/Private/InstallBundleTypes.cpp
@@ -2,10 +2,13 @@
 
 #include "InstallBundleTypes.h"
 
+#include "Containers/ContainerAllocationPolicies.h"
 #include "InstallBundleManagerPrivate.h"
 #include "InstallBundleUtils.h"
+#include "InstallBundleManagerInterface.h"
 #include "Misc/CString.h"
 
+PRAGMA_DISABLE_DEPRECATION_WARNINGS
 const TCHAR* LexToString(EInstallBundleSourceType Type)
 {
 	static const TCHAR* Strings[] =
@@ -20,9 +23,13 @@ const TCHAR* LexToString(EInstallBundleSourceType Type)
 		TEXT("Streaming"),
 	};
 
-	return InstallBundleUtil::TLexToString(Type, Strings);
+	// Clang has issues with the not silencing deprecation warnings for TLexToString here,
+	// so explicitly pass the types to make it happy.
+	return InstallBundleUtil::TLexToString<EInstallBundleSourceType, decltype(Strings), EInstallBundleSourceType::Count>(Type, Strings);
 }
+PRAGMA_ENABLE_DEPRECATION_WARNINGS
 
+PRAGMA_DISABLE_DEPRECATION_WARNINGS
 void LexFromString(EInstallBundleSourceType& OutType, const TCHAR* String)
 {
 	OutType = EInstallBundleSourceType::Count;
@@ -37,6 +44,74 @@ void LexFromString(EInstallBundleSourceType& OutType, const TCHAR* String)
 		}
 	}
 }
+PRAGMA_ENABLE_DEPRECATION_WARNINGS
+
+class FInstallBundleSourceTypeNameTable
+{
+public:
+	FInstallBundleSourceTypeNameTable()
+	{
+		check(IsInGameThread());
+
+		// Find all possible sources from config
+		TArray<FString> ConfgSources;
+		TMap<FString, FString> ConfigFallbackSources;
+		if (!InstallBundleUtil::GetConfiguredBundleSources(ConfgSources, ConfigFallbackSources))
+		{
+			return;
+		}
+
+		for (FString& Source : ConfgSources)
+		{
+			NameTable.AddUnique(MoveTemp(Source));
+		}
+
+		for (TPair<FString, FString>& Pair : ConfigFallbackSources)
+		{
+			NameTable.AddUnique(MoveTemp(Pair.Key));
+			NameTable.AddUnique(MoveTemp(Pair.Value));
+		}
+	}
+
+	FStringView FindBundleSourceTypeByName(FStringView InName) const
+	{
+		for (const FString& Str : NameTable)
+		{
+			if (Str == InName)
+			{
+				return Str;
+			}
+		}
+
+		return FStringView(TEXTVIEW(""));
+	}
+
+private:
+	TArray<FString, TInlineAllocator<8>> NameTable;
+};
+
+static const FInstallBundleSourceTypeNameTable& GetInstallBundleSourceTypeNameTable()
+{
+	static FInstallBundleSourceTypeNameTable InstallBundleSourceTypeNameTable;
+	return InstallBundleSourceTypeNameTable;
+}
+
+FInstallBundleSourceType::FInstallBundleSourceType(FStringView InNameStr)
+	: NameStr(GetInstallBundleSourceTypeNameTable().FindBundleSourceTypeByName(InNameStr))
+{
+}
+
+PRAGMA_DISABLE_DEPRECATION_WARNINGS
+FInstallBundleSourceType::FInstallBundleSourceType(EInstallBundleSourceType InLegacySourceType)
+	: NameStr(GetInstallBundleSourceTypeNameTable().FindBundleSourceTypeByName(LexToString(InLegacySourceType)))
+{
+}
+PRAGMA_ENABLE_DEPRECATION_WARNINGS
+
+const TCHAR* LexToString(FInstallBundleSourceType Type)
+{
+	return Type.GetNameCStr();
+}
 
 const TCHAR* LexToString(EInstallBundleManagerInitResult Result)
 {
diff --git a/Engine/Source/Runtime/InstallBundleManager/Private/InstallBundleUtils.cpp b/Engine/Source/Runtime/InstallBundleManager/Private/InstallBundleUtils.cpp
--- a/Engine/Source/Runtime/InstallBundleManager/Private/InstallBundleUtils.cpp
+++ b/Engine/Source/Runtime/InstallBundleManager/Private/InstallBundleUtils.cpp
@@ -56,6 +56,54 @@ namespace InstallBundleUtil
 		return Prefix;
 	}
 
+	bool GetConfiguredBundleSources(TArray<FString>& OutSources, TMap<FString, FString>& OutFallbackSources)
+	{
+#ifdef INSTALL_BUNDLE_SOURCES_FALLBACK_CONFIG_SECTION
+		const TCHAR* BundleSourceFallbackSection = TEXT(INSTALL_BUNDLE_SOURCES_FALLBACK_CONFIG_SECTION);
+#else
+		const TCHAR* BundleSourceFallbackSection = TEXT("InstallBundleManager.FallbackBundleSources");
+#endif // INSTALL_BUNDLE_SOURCES_FALLBACK_CONFIG_SECTION
+
+		TArray<FString> ConfigBundleFallbacks;
+		if (GConfig->GetArray(BundleSourceFallbackSection, TEXT("FallbackBundleSources"), ConfigBundleFallbacks, GInstallBundleIni))
+		{
+			OutFallbackSources.Empty(ConfigBundleFallbacks.Num());
+			for (FString& ConfigFallback : ConfigBundleFallbacks)
+			{
+				// Remove parentheses
+				ConfigFallback.ReplaceInline(TEXT("("), TEXT(""));
+				ConfigFallback.ReplaceInline(TEXT(")"), TEXT(""));
+
+				TArray<FString> Tokens;
+				if (2 != ConfigFallback.ParseIntoArrayWS(Tokens, TEXT(",")))
+				{
+					ensureAlwaysMsgf(false, TEXT("Malformed entry in InstallBundleManager.FallbackBundleSources"));
+					return false;
+				}
+
+				OutFallbackSources.Add(MoveTemp(Tokens[0]), MoveTemp(Tokens[1]));
+			}
+		}
+
+#ifdef INSTALL_BUNDLE_SOURCES_CONFIG_SECTION
+		const TCHAR* BundleSourceSection = TEXT(INSTALL_BUNDLE_SOURCES_CONFIG_SECTION);
+#else
+		const TCHAR* BundleSourceSection = TEXT("InstallBundleManager.BundleSources");
+#endif // INSTALL_BUNDLE_SOURCE_CONFIG_SECTION
+
+		TArray<FString> ConfigBundleSources;
+		if (ensureAlways(GConfig->GetArray(BundleSourceSection, TEXT("DefaultBundleSources"), ConfigBundleSources, GInstallBundleIni)))
+		{
+			OutSources = MoveTemp(ConfigBundleSources);
+		}
+		else
+		{
+			return false;
+		}
+
+		return true;
+	}
+
 	bool HasInstallBundleInConfig(const FString& BundleName)
 	{
 		const FConfigFile* InstallBundleConfig = GConfig->FindConfigFile(GInstallBundleIni);
diff --git a/Engine/Source/Runtime/InstallBundleManager/Public/InstallBundleCache.h b/Engine/Source/Runtime/InstallBundleManager/Public/InstallBundleCache.h
--- a/Engine/Source/Runtime/InstallBundleManager/Public/InstallBundleCache.h
+++ b/Engine/Source/Runtime/InstallBundleManager/Public/InstallBundleCache.h
@@ -41,13 +41,13 @@ enum class EInstallBundleCacheReserveResult : int8
 
 struct FInstallBundleCacheReserveResult
 {
-	TMap<FName, TArray<EInstallBundleSourceType>> BundlesToEvict;
+	TMap<FName, TArray<FInstallBundleSourceType>> BundlesToEvict;
 	EInstallBundleCacheReserveResult Result = EInstallBundleCacheReserveResult::Success;
 };
 
 struct FInstallBundleCacheFlushResult
 {
-	TMap<FName, TArray<EInstallBundleSourceType>> BundlesToEvict;
+	TMap<FName, TArray<FInstallBundleSourceType>> BundlesToEvict;
 };
 
 class FInstallBundleCache : public TSharedFromThis<FInstallBundleCache>
@@ -60,12 +60,12 @@ public:
 	FName GetName() const { return CacheName; }
 
 	// Add a bundle to the cache.  
-	INSTALLBUNDLEMANAGER_API void AddOrUpdateBundle(EInstallBundleSourceType Source, const FInstallBundleCacheBundleInfo& AddInfo);
+	INSTALLBUNDLEMANAGER_API void AddOrUpdateBundle(FInstallBundleSourceType Source, const FInstallBundleCacheBundleInfo& AddInfo);
 
-	INSTALLBUNDLEMANAGER_API void RemoveBundle(EInstallBundleSourceType Source, FName BundleName);
+	INSTALLBUNDLEMANAGER_API void RemoveBundle(FInstallBundleSourceType Source, FName BundleName);
 
 	INSTALLBUNDLEMANAGER_API TOptional<FInstallBundleCacheBundleInfo> GetBundleInfo(FName BundleName) const;
-	INSTALLBUNDLEMANAGER_API TOptional<FInstallBundleCacheBundleInfo> GetBundleInfo(EInstallBundleSourceType Source, FName BundleName) const;	
+	INSTALLBUNDLEMANAGER_API TOptional<FInstallBundleCacheBundleInfo> GetBundleInfo(FInstallBundleSourceType Source, FName BundleName) const;	
 
 	// Return the total size of the cache
 	INSTALLBUNDLEMANAGER_API uint64 GetSize() const;
@@ -79,10 +79,10 @@ public:
 	INSTALLBUNDLEMANAGER_API FInstallBundleCacheReserveResult Reserve(FName BundleName);
 
 	// Called from bundle manager, returns all bundles that can be evicted
-	INSTALLBUNDLEMANAGER_API FInstallBundleCacheFlushResult Flush(EInstallBundleSourceType* Source = nullptr);
+	INSTALLBUNDLEMANAGER_API FInstallBundleCacheFlushResult Flush(FInstallBundleSourceType* Source = nullptr);
 
 	INSTALLBUNDLEMANAGER_API bool Contains(FName BundleName) const;
-	INSTALLBUNDLEMANAGER_API bool Contains(EInstallBundleSourceType Source, FName BundleName) const;
+	INSTALLBUNDLEMANAGER_API bool Contains(FInstallBundleSourceType Source, FName BundleName) const;
 
 	INSTALLBUNDLEMANAGER_API bool IsReserved(FName BundleName) const;
 
@@ -172,7 +172,7 @@ private:
 
 private:
 
-	TMap<FName, TMap<EInstallBundleSourceType, FPerSourceBundleCacheInfo>> PerSourceCacheInfo;
+	TMap<FName, TMap<FInstallBundleSourceType, FPerSourceBundleCacheInfo>> PerSourceCacheInfo;
 
 	// mutable to allow sorting in const contexts
 	mutable TMap<FName, FBundleCacheInfo> CacheInfo;
diff --git a/Engine/Source/Runtime/InstallBundleManager/Public/InstallBundleManagerInterface.h b/Engine/Source/Runtime/InstallBundleManager/Public/InstallBundleManagerInterface.h
--- a/Engine/Source/Runtime/InstallBundleManager/Public/InstallBundleManagerInterface.h
+++ b/Engine/Source/Runtime/InstallBundleManager/Public/InstallBundleManagerInterface.h
@@ -88,7 +88,7 @@ enum class EInstallBundleManagerInitErrorHandlerResult
 	StopInitialization, // Stop trying to initialize
 };
 
-using FInstallBundleSourceOrCache = TUnion<EInstallBundleSourceType, FName>;
+using FInstallBundleSourceOrCache = TUnion<FInstallBundleSourceType, FName>;
 
 DECLARE_DELEGATE_RetVal_OneParam(EInstallBundleManagerInitErrorHandlerResult, FInstallBundleManagerInitErrorHandler, EInstallBundleManagerInitResult);
 DECLARE_MULTICAST_DELEGATE_OneParam(FInstallBundleManagerInitCompleteMultiDelegate, EInstallBundleManagerInitResult);
@@ -121,9 +121,9 @@ public:
 
 	virtual void Initialize() {}
 
-	virtual bool HasBundleSource(EInstallBundleSourceType SourceType) const = 0;
+	virtual bool HasBundleSource(FInstallBundleSourceType SourceType) const = 0;
 
-	INSTALLBUNDLEMANAGER_API virtual const TSharedPtr<IInstallBundleSource> GetBundleSource(EInstallBundleSourceType SourceType) const;
+	INSTALLBUNDLEMANAGER_API virtual const TSharedPtr<IInstallBundleSource> GetBundleSource(FInstallBundleSourceType SourceType) const;
 
 
 	virtual FDelegateHandle PushInitErrorCallback(FInstallBundleManagerInitErrorHandler Callback) = 0;
diff --git a/Engine/Source/Runtime/InstallBundleManager/Public/InstallBundleSourceInterface.h b/Engine/Source/Runtime/InstallBundleManager/Public/InstallBundleSourceInterface.h
--- a/Engine/Source/Runtime/InstallBundleManager/Public/InstallBundleSourceInterface.h
+++ b/Engine/Source/Runtime/InstallBundleManager/Public/InstallBundleSourceInterface.h
@@ -26,7 +26,7 @@ public:
 	virtual ~IInstallBundleSource() {}
 
 	// Returns a unique id for this source
-	virtual EInstallBundleSourceType GetSourceType() const = 0;
+	virtual FInstallBundleSourceType GetSourceType() const = 0;
 
 	// Returns the how this source should be weighted when combined with other sources
 	virtual float GetSourceWeight() const { return 1.0f; }
diff --git a/Engine/Source/Runtime/InstallBundleManager/Public/InstallBundleTypes.h b/Engine/Source/Runtime/InstallBundleManager/Public/InstallBundleTypes.h
--- a/Engine/Source/Runtime/InstallBundleManager/Public/InstallBundleTypes.h
+++ b/Engine/Source/Runtime/InstallBundleManager/Public/InstallBundleTypes.h
@@ -20,7 +20,7 @@
 	#define WITH_PLATFORM_INSTALL_BUNDLE_SOURCE 0
 #endif
 
-enum class EInstallBundleSourceType : int
+enum class UE_DEPRECATED(5.5, "Use FInstallBundleSourceType") EInstallBundleSourceType : int
 {
 	Bulk,
 	Launcher,
@@ -32,9 +32,46 @@ enum class EInstallBundleSourceType : int
 	Streaming,
 	Count,
 };
+PRAGMA_DISABLE_DEPRECATION_WARNINGS
 ENUM_RANGE_BY_COUNT(EInstallBundleSourceType, EInstallBundleSourceType::Count);
+UE_DEPRECATED(5.5, "Use FInstallBundleSourceType")
 INSTALLBUNDLEMANAGER_API const TCHAR* LexToString(EInstallBundleSourceType Type);
+UE_DEPRECATED(5.5, "Use FInstallBundleSourceType")
 INSTALLBUNDLEMANAGER_API void LexFromString(EInstallBundleSourceType& OutType, const TCHAR* String);
+PRAGMA_ENABLE_DEPRECATION_WARNINGS
+
+
+class FInstallBundleSourceType
+{
+private:
+	FStringView NameStr;
+
+public:
+	INSTALLBUNDLEMANAGER_API explicit FInstallBundleSourceType(FStringView InNameStr);
+
+PRAGMA_DISABLE_DEPRECATION_WARNINGS
+	INSTALLBUNDLEMANAGER_API FInstallBundleSourceType(EInstallBundleSourceType InLegacySourceType);
+PRAGMA_ENABLE_DEPRECATION_WARNINGS
+
+	FStringView GetName() const { return NameStr; }
+	FString GetNameStr() const { return FString(NameStr); }
+	const TCHAR* GetNameCStr() const { return NameStr.GetData(); }
+
+	bool IsValid() const { return !NameStr.IsEmpty(); }
+
+	bool operator==(const FInstallBundleSourceType Other) const
+	{
+		// These should always point to constant strings
+		return NameStr.GetData() == Other.NameStr.GetData();
+	}
+
+	friend inline uint32 GetTypeHash(FInstallBundleSourceType In)
+	{
+		// These should always point to constant strings
+		return PointerHash(In.NameStr.GetData());
+	}
+};
+INSTALLBUNDLEMANAGER_API const TCHAR* LexToString(FInstallBundleSourceType Type);
 
 enum class EInstallBundleManagerInitState : int
 {
@@ -84,13 +121,13 @@ struct FInstallBundleContentState
 {
 	EInstallBundleInstallState State = EInstallBundleInstallState::NotInstalled;
 	float Weight = 0.0f;
-	TMap<EInstallBundleSourceType, FString> Version;
+	TMap<FInstallBundleSourceType, FString> Version;
 };
 
 struct FInstallBundleCombinedContentState
 {
 	TMap<FName, FInstallBundleContentState> IndividualBundleStates;
-	TMap<EInstallBundleSourceType, FString> CurrentVersion;
+	TMap<FInstallBundleSourceType, FString> CurrentVersion;
 	uint64 DownloadSize = 0;
 	uint64 InstallSize = 0;
 	uint64 InstallOverheadSize = 0;
diff --git a/Engine/Source/Runtime/InstallBundleManager/Public/InstallBundleUtils.h b/Engine/Source/Runtime/InstallBundleManager/Public/InstallBundleUtils.h
--- a/Engine/Source/Runtime/InstallBundleManager/Public/InstallBundleUtils.h
+++ b/Engine/Source/Runtime/InstallBundleManager/Public/InstallBundleUtils.h
@@ -37,6 +37,8 @@ namespace InstallBundleUtil
 
 	INSTALLBUNDLEMANAGER_API const FString& GetInstallBundleSectionPrefix();
 
+	INSTALLBUNDLEMANAGER_API bool GetConfiguredBundleSources(TArray<FString>& OutSources, TMap<FString, FString>& OutFallbackSources);
+
 	// returns true if the given BundleName exists in the InstallBundle.ini config
 	INSTALLBUNDLEMANAGER_API bool HasInstallBundleInConfig(const FString& BundleName);
 
@@ -65,7 +67,7 @@ namespace InstallBundleUtil
 	constexpr auto& CastAsUnderlying(EnumType &Type)
 	{
 		static_assert(TIsEnum<EnumType>::Value, "");
-		using UnderType = __underlying_type(EnumType);
+		using UnderType = std::underlying_type_t<EnumType>;
 		return *reinterpret_cast<UnderType*>(&Type);
 	}
 
@@ -73,7 +75,7 @@ namespace InstallBundleUtil
 	constexpr const auto& CastAsUnderlying(const EnumType &Type)
 	{
 		static_assert(TIsEnum<EnumType>::Value, "");
-		using UnderType = __underlying_type(EnumType);
+		using UnderType = std::underlying_type_t<EnumType>;
 		return *reinterpret_cast<const UnderType*>(&Type);
 	}
 
@@ -81,14 +83,14 @@ namespace InstallBundleUtil
 	constexpr auto CastToUnderlying(EnumType Type)
 	{
 		static_assert(TIsEnum<EnumType>::Value, "");
-		using UnderType = __underlying_type(EnumType);
+		using UnderType = std::underlying_type_t<EnumType>;
 		return static_cast<UnderType>(Type);
 	}
 
-	template<typename EnumType, typename StringArrType>
+	template<typename EnumType, typename StringArrType, EnumType EnumCount = EnumType::Count>
 	const TCHAR* TLexToString(EnumType E, const StringArrType& Strings)
 	{
-		constexpr auto Count = InstallBundleUtil::CastToUnderlying(EnumType::Count);
+		constexpr auto Count = InstallBundleUtil::CastToUnderlying(EnumCount);
 		static_assert(Count == UE_ARRAY_COUNT(Strings), "");
 
 		auto Idx = InstallBundleUtil::CastToUnderlying(E);
@@ -634,7 +636,7 @@ namespace InstallBundleUtil
 		FContentRequestSharedContext& operator=(const FContentRequestSharedContext& Other) = delete;
 		FContentRequestSharedContext& operator=(FContentRequestSharedContext&& Other) = default;
 
-		TMap<EInstallBundleSourceType, TUniquePtr<IBundleSourceContentRequestSharedContext>> BundleSourceSharedContext;
+		TMap<FInstallBundleSourceType, TUniquePtr<IBundleSourceContentRequestSharedContext>> BundleSourceSharedContext;
 	};
 	using FContentRequestSharedContextPtr = TSharedPtr<FContentRequestSharedContext>;
 }
diff --git a/Engine/Source/Runtime/NullInstallBundleManager/Private/NullInstallBundleManager.cpp b/Engine/Source/Runtime/NullInstallBundleManager/Private/NullInstallBundleManager.cpp
--- a/Engine/Source/Runtime/NullInstallBundleManager/Private/NullInstallBundleManager.cpp
+++ b/Engine/Source/Runtime/NullInstallBundleManager/Private/NullInstallBundleManager.cpp
@@ -4,9 +4,9 @@
 #include "InstallBundleManagerModule.h"
 #include "Modules/ModuleManager.h"
 
-class FNullInstallBundleManager : public IInstallBundleManager
+class FNullInstallBundleManager final : public IInstallBundleManager
 {
-	virtual bool HasBundleSource(EInstallBundleSourceType SourceType) const override
+	virtual bool HasBundleSource(FInstallBundleSourceType SourceType) const override
 	{
 		return false;
 	}
