Install Bundle Manager
- Handle the case that a bundle source removed content even if EInstallBundleReleaseRequestFlags::RemoveFilesIfPossible was not set.

diff --git a/Engine/Plugins/Experimental/DefaultInstallBundleManager/Source/Private/DefaultInstallBundleManager.cpp b/Engine/Plugins/Experimental/DefaultInstallBundleManager/Source/Private/DefaultInstallBundleManager.cpp
--- a/Engine/Plugins/Experimental/DefaultInstallBundleManager/Source/Private/DefaultInstallBundleManager.cpp
+++ b/Engine/Plugins/Experimental/DefaultInstallBundleManager/Source/Private/DefaultInstallBundleManager.cpp
@@ -1434,6 +1434,13 @@ void FDefaultInstallBundleManager::UpdateBundleSourceReleaseComplete(TSharedRef<
 			Request->Result = ResultInfo.Result;
 		}
 
+		if (ResultInfo.bContentWasRemoved)
+		{
+			// Removing content wasn't requested, but the source did it anyway so deal with it
+			bContentWasRemoved = true;
+			continue;
+		}
+
 		// Since we didn't remove content, update last access times in any caches this bundle participates in
 		if (const FName* CacheName = BundleSourceCaches.Find(Pair.Key))
 		{
@@ -3526,6 +3533,9 @@ TValueOrError<FInstallBundleReleaseRequestInfo, EInstallBundleResult> FDefaultIn
 		{
 			ActiveQueuedRequest->Flags |= EInstallBundleReleaseRequestFlags::RemoveFilesIfPossible;
 
+			// TODO: This assumes that a bundle source that doesn't remove files will always immediatly callback on release.
+			// That is probably a bad assumption, we can't control how the bundle source may be written.
+			// It would be safer to instead in enqueue this with a prereq that there is no pending release
 #if DO_CHECK
 			// Since a request without RemoveFilesIfPossible does no async work after unmounting, it shouldn't be possible to 
 			// call RequestReleaseContent and find an ActiveQueuedRequest that is past the unmounting step.
