Updates to chunk installer wrapper bundle source to bring it to parity with platform-native bundle sources, including:

 - initial/default named chunk is included in the list of known named chunks, and applied to every non-platform bundle.
 - non-platform bundles are included in the GetContentState results, flagged as 'not installed'
 - more forgiving language chunks on some platforms.
 - use bespoke log prefix.
 - use bespoke correct init-complete telemetry name.

#rnx

diff --git a/Engine/Plugins/Experimental/DefaultInstallBundleManager/Source/Private/InstallBundleManagerUtil.cpp b/Engine/Plugins/Experimental/DefaultInstallBundleManager/Source/Private/InstallBundleManagerUtil.cpp
--- a/Engine/Plugins/Experimental/DefaultInstallBundleManager/Source/Private/InstallBundleManagerUtil.cpp
+++ b/Engine/Plugins/Experimental/DefaultInstallBundleManager/Source/Private/InstallBundleManagerUtil.cpp
@@ -664,6 +664,17 @@ namespace InstallBundleManagerAnalytics
 			TEXT("InitResultString"), InitResultString));
 	}
 
+	void FireEvent_InitBundleSourcePlatformChunkInstallComplete(IAnalyticsProviderET* AnalyticsProvider, const FString InitResultString)
+	{
+		if (AnalyticsProvider == nullptr || InstallBundleUtil::FInstallBundleSuppressAnalytics::IsEnabled())
+		{
+			return;
+		}
+
+		AnalyticsProvider->RecordEvent(TEXT("InstallBundleManager.InitBundleSourcePlatformChunkInstallComplete"), MakeAnalyticsEventAttributeArray(
+			TEXT("InitResultString"), InitResultString));
+	}
+
 	void FireEvent_BundleLatestClientCheckComplete(IAnalyticsProviderET* AnalyticsProvider, 
 		const FString& BundleName, 
 		bool bSkippedCheck, 
diff --git a/Engine/Plugins/Experimental/DefaultInstallBundleManager/Source/Private/InstallBundleSourcePlatformChunkInstall.cpp b/Engine/Plugins/Experimental/DefaultInstallBundleManager/Source/Private/InstallBundleSourcePlatformChunkInstall.cpp
--- a/Engine/Plugins/Experimental/DefaultInstallBundleManager/Source/Private/InstallBundleSourcePlatformChunkInstall.cpp
+++ b/Engine/Plugins/Experimental/DefaultInstallBundleManager/Source/Private/InstallBundleSourcePlatformChunkInstall.cpp
@@ -14,9 +14,9 @@
 #include "IPlatformFilePak.h"
 #include "Stats/Stats.h"
 
-#define LOG_SOURCE_CHUNKINSTALL(Verbosity, Format, ...) LOG_INSTALL_BUNDLE_MAN(Verbosity, TEXT("InstallBundleSourceIntelligentDelivery: ") Format, ##__VA_ARGS__)
+#define LOG_SOURCE_CHUNKINSTALL(Verbosity, Format, ...) LOG_INSTALL_BUNDLE_MAN(Verbosity, TEXT("InstallBundleSourcePlatformChunkInstall: ") Format, ##__VA_ARGS__)
 
-#define LOG_SOURCE_CHUNKINSTALL_OVERRIDE(VerbosityOverride, Verbosity, Format, ...) LOG_INSTALL_BUNDLE_MAN_OVERRIDE(VerbosityOverride, Verbosity, TEXT("InstallBundleSourceIntelligentDelivery: ") Format, ##__VA_ARGS__)
+#define LOG_SOURCE_CHUNKINSTALL_OVERRIDE(VerbosityOverride, Verbosity, Format, ...) LOG_INSTALL_BUNDLE_MAN_OVERRIDE(VerbosityOverride, Verbosity, TEXT("InstallBundleSourcePlatformChunkInstall: ") Format, ##__VA_ARGS__)
 
 
 
@@ -180,6 +180,10 @@ void FInstallBundleSourcePlatformChunkInstall::AsyncInit(FInstallBundleSourceIni
 			return;
 		}
 
+		// read the default chunk name. this is used for bundles that have no platform data configured
+		FString DefaultPlatformChunkName;
+		GConfig->GetString(TEXT("InstallBundleSource.Platform.MiscSettings"), TEXT("DefaultPlatformChunkName"), DefaultPlatformChunkName, GInstallBundleIni);
+
 		// cache all named chunks
 		Context->NamedChunks.Append( Context->PlatformChunkInstall->GetNamedChunksByType(ENamedChunkType::OnDemand) );
 		Context->NamedChunks.Append( Context->PlatformChunkInstall->GetNamedChunksByType(ENamedChunkType::Language) );
@@ -198,16 +202,24 @@ void FInstallBundleSourcePlatformChunkInstall::AsyncInit(FInstallBundleSourceIni
 			{
 				continue;
 			}
-
-			FString PlatformChunkName;
-			InstallBundleConfig->GetString(*Section, TEXT("PlatformChunkName"), PlatformChunkName);
 			
+			// read the PlatformChunkID. Note this is only used for checking disabled or defaulted bundles (such as Game Feature Plugins)
+			// it isn't actually used by this bundle source. At some point a PlatformChunkName-based equivalent will be added instead
 			int32 PlatformChunkID = 0;
-			if (PlatformChunkName.IsEmpty() && InstallBundleConfig->GetInt(*Section, TEXT("PlatformChunkID"), PlatformChunkID) && PlatformChunkID < 0)
+			InstallBundleConfig->GetInt(*Section, TEXT("PlatformChunkID"), PlatformChunkID);
+			if (PlatformChunkID < 0)
 			{
 				continue;
 			}
-			//... NB PlatformChunkID is ignored by this bundle source, and only used to mark the bundle as a platform bundle
+			
+			FString PlatformChunkName;
+			InstallBundleConfig->GetString(*Section, TEXT("PlatformChunkName"), PlatformChunkName);
+			if (PlatformChunkID == 0 && PlatformChunkName.IsEmpty())
+			{
+				// This bundle has no specific PlatformChunkName and the PlatformChunkID is either 0 ("default") or unspecified, so use the name of the default chunk name.
+				// The code below will assume the chunk name is the same as the bundle name if it's not been specified.
+				PlatformChunkName = DefaultPlatformChunkName;
+			}
 
 			// create bundle info
 			FName BundleName( *Section.RightChop(InstallBundleUtil::GetInstallBundleSectionPrefix().Len()));
@@ -239,7 +251,7 @@ void FInstallBundleSourcePlatformChunkInstall::AsyncInit(FInstallBundleSourceIni
 							ChunkFilePath = TEXT("../../../") + ChunkFilePath;
 						}
 						FPaths::NormalizeFilename(ChunkFilePath);
-						NamedChunkFilePaths.FindOrAdd(BundleName).Add(ChunkFilePath);
+						NamedChunkFilePaths.FindOrAdd(BundleInfo.NamedChunk).Add(ChunkFilePath);
 					}
 				}
 			}
@@ -311,7 +323,7 @@ void FInstallBundleSourcePlatformChunkInstall::AsyncInit(FInstallBundleSourceIni
 		}
 
 		LOG_SOURCE_CHUNKINSTALL(Display, TEXT("Fire Init Analytic: %s"), LexToString(InitInfo.Result));
-		InstallBundleManagerAnalytics::FireEvent_InitBundleSourceIntelligentDeliveryComplete(AnalyticsProvider.Get(), LexToString(InitInfo.Result));
+		InstallBundleManagerAnalytics::FireEvent_InitBundleSourcePlatformChunkInstallComplete(AnalyticsProvider.Get(), LexToString(InitInfo.Result));
 
 		OnInitCompleteCallback.Execute(AsShared(), MoveTemp(InitInfo));
 	};
@@ -367,38 +379,44 @@ void FInstallBundleSourcePlatformChunkInstall::GetContentState(TArrayView<const
 		if (BundleInfo != nullptr)
 		{
 			FChunkInstallationStatusDetail ChunkStatusDetail;
-			if (PlatformChunkInstall->GetNamedChunkInstallationStatus( BundleInfo->NamedChunk, ChunkStatusDetail ))
-			{
-				FInstallBundleContentState& IndividualBundleState = State.IndividualBundleStates.Add(BundleName);
+			bool bHasChunkStatus = PlatformChunkInstall->GetNamedChunkInstallationStatus( BundleInfo->NamedChunk, ChunkStatusDetail );
 
-				if (ChunkStatusDetail.bIsInstalled)
-				{
-					IndividualBundleState.State = EInstallBundleInstallState::UpToDate;
-					IndividualBundleState.Version.Add(GetSourceType(), InstallBundleUtil::GetAppVersion());
-				}
-				else
-				{
-					IndividualBundleState.State = EInstallBundleInstallState::NotInstalled;
-					IndividualBundleState.Version.Add(GetSourceType(), TEXT(""));
-				}
+			FInstallBundleContentState& IndividualBundleState = State.IndividualBundleStates.Add(BundleName);
+			if (bHasChunkStatus && ChunkStatusDetail.bIsInstalled)
+			{
+				IndividualBundleState.State = EInstallBundleInstallState::UpToDate;
+				IndividualBundleState.Version.Add(GetSourceType(), InstallBundleUtil::GetAppVersion());
+			}
+			else
+			{
+				// named chunk is either not installed, or it's not a valid named chunk (also meaning it isn't installed)... 
+				// an invalid named chunk likely means there is an entry for it in the bundle ini, but the platform chunk it refers to does not exist
+				IndividualBundleState.State = EInstallBundleInstallState::NotInstalled;
+				IndividualBundleState.Version.Add(GetSourceType(), TEXT(""));
+			}
 
+			uint64 RemainingDownloadSize = 0;
+			if (bHasChunkStatus)
+			{
 				check( ChunkStatusDetail.CurrentInstallSize <= ChunkStatusDetail.FullInstallSize);
-				uint64 RemainingDownloadSize = (ChunkStatusDetail.FullInstallSize - ChunkStatusDetail.CurrentInstallSize);
-
+				RemainingDownloadSize = (ChunkStatusDetail.FullInstallSize - ChunkStatusDetail.CurrentInstallSize);
 				TotalRemainingDownloadSize += RemainingDownloadSize;
-				BundleToRemaingDownloadSize.Add(BundleName, RemainingDownloadSize);
 			}
+			BundleToRemaingDownloadSize.Add(BundleName, RemainingDownloadSize);
 		}
 	}
 
-	// compute the download weight for all of the bundles - higher weight is a bigger download
-	for (TTuple<FName,uint64> Pair : BundleToRemaingDownloadSize)
+	// compute the download weight for all of the known bundles - higher weight is a bigger download
+	for (TTuple<FName,FInstallBundleContentState>& BundleStatePair : State.IndividualBundleStates)
 	{
-		uint64 RemainingDownloadSize = Pair.Value;
-		double Weight = (TotalRemainingDownloadSize > 0)  ?  ((double)RemainingDownloadSize / (double)TotalRemainingDownloadSize)  :  0.0;
+		if (TotalRemainingDownloadSize == 0)
+		{
+			BundleStatePair.Value.Weight = 1.0f / BundleNames.Num();
+			continue;
+		}
 
-		FInstallBundleContentState& IndividualBundleState = State.IndividualBundleStates.FindChecked(Pair.Key);
-		IndividualBundleState.Weight = FMath::Max(Weight, InstallBundleUtil::MinimumBundleWeight); // this Max() does mean the total weight will be > 1.0 but all other bundle sources do it this way too
+		double Weight = (double)BundleToRemaingDownloadSize[BundleStatePair.Key] / (double)TotalRemainingDownloadSize;
+		BundleStatePair.Value.Weight = FMath::Max(Weight, InstallBundleUtil::MinimumBundleWeight); // this Max() does mean the total weight will be > 1.0 but all other bundle sources do it this way too
 	}
 
 	// Don't return any size info since all the size should be reserved by the system	
diff --git a/Engine/Plugins/Experimental/DefaultInstallBundleManager/Source/Public/InstallBundleManagerUtil.h b/Engine/Plugins/Experimental/DefaultInstallBundleManager/Source/Public/InstallBundleManagerUtil.h
--- a/Engine/Plugins/Experimental/DefaultInstallBundleManager/Source/Public/InstallBundleManagerUtil.h
+++ b/Engine/Plugins/Experimental/DefaultInstallBundleManager/Source/Public/InstallBundleManagerUtil.h
@@ -263,6 +263,16 @@ namespace InstallBundleManagerAnalytics
 	DEFAULTINSTALLBUNDLEMANAGER_API void FireEvent_InitBundleSourceIntelligentDeliveryComplete(IAnalyticsProviderET* AnalyticsProvider,
 		const FString InitResultString);
 
+	/**
+	* @EventName InstallBundleManager.FireEvent_InitBundleSourcePlatformChunkInstallComplete
+	* @Trigger Bundle Manager finished async initialization
+	* @Type Client
+	* @EventParam InitResultString String Result code
+	* @Comments
+	*/
+	DEFAULTINSTALLBUNDLEMANAGER_API void FireEvent_InitBundleSourcePlatformChunkInstallComplete(IAnalyticsProviderET* AnalyticsProvider,
+		const FString InitResultString);
+
 	/**
 	 * @EventName InstallBundleManager.BundleLatestClientCheckComplete
 	 *
