Fix GFPs using InstallBundle protocol releasing pinned content too early
- Add EInstallBundleReleaseRequestFlags::SkipReleaseUnmountOnly
- Update relevant analytics
- only trigger ensure about -UsePaks if Paks will actually be mounted

diff --git a/Engine/Plugins/Experimental/DefaultInstallBundleManager/Source/Private/DefaultInstallBundleManager.cpp b/Engine/Plugins/Experimental/DefaultInstallBundleManager/Source/Private/DefaultInstallBundleManager.cpp
--- a/Engine/Plugins/Experimental/DefaultInstallBundleManager/Source/Private/DefaultInstallBundleManager.cpp
+++ b/Engine/Plugins/Experimental/DefaultInstallBundleManager/Source/Private/DefaultInstallBundleManager.cpp
@@ -1324,6 +1324,13 @@ void FDefaultInstallBundleManager::UpdateBundleSources(FContentReleaseRequestRef
 		return;
 	}
 
+	if (EnumHasAnyFlags(Request->Flags, EInstallBundleReleaseRequestFlags::SkipReleaseUnmountOnly))
+	{
+		LOG_INSTALL_BUNDLE_MAN_OVERRIDE(Request->LogVerbosityOverride, Display, TEXT("Skipping Updated Sources for Release Request %s"), *Request->BundleName.ToString());
+		Request->StepResult = EContentRequestStepResult::Done;
+		return;
+	}
+
 	// Release from any caches that were reserved
 	for (const TPair<FName, TSharedRef<FInstallBundleCache>>& Pair : BundleCaches)
 	{
@@ -1545,18 +1552,16 @@ void FDefaultInstallBundleManager::MountPaks(FContentRequestRef Request)
 bool FDefaultInstallBundleManager::MountPaksInList(TArrayView<FString> Paths, ELogVerbosity::Type LogVerbosityOverride)
 {
 	FPakPlatformFile* PakPlatformFile = static_cast<FPakPlatformFile*>(FPlatformFileManager::Get().FindPlatformFile(TEXT("PakFile")));
-	if (!PakPlatformFile)
-	{
-		ensureMsgf(false, TEXT("Pak files have not been correctly initalized. Use -UsePaks on the cmdline if you are using the UnrealEditor.exe"));
-		return false; // if FCoreDelegates::MountPak is unbound there is a major issue.
-	}
 
 	// Sort in descending order.
 	Paths.Sort(TGreater<FString>());
 
 	// Find already mounted paks
 	TSet<FString> MountedPaks;
-	PakPlatformFile->GetMountedPakFilenames(MountedPaks);
+	if (PakPlatformFile)
+	{
+		PakPlatformFile->GetMountedPakFilenames(MountedPaks);
+	}
 
 	bool bMountedPaks = false;
 	for (const FString& File : Paths)
@@ -1572,6 +1577,12 @@ bool FDefaultInstallBundleManager::MountPaksInList(TArrayView<FString> Paths, EL
 			continue;
 		}
 
+		if (!PakPlatformFile)
+		{
+			ensureMsgf(false, TEXT("Pak files have not been correctly initalized. Use -UsePaks on the cmdline if you are using the UnrealEditor.exe"));
+			return false; // if FCoreDelegates::MountPak is unbound there is a major issue.
+		}
+
 		LOG_INSTALL_BUNDLE_MAN_OVERRIDE(LogVerbosityOverride, Verbose, TEXT("Mounting pak file: %s \n"), *File);
 
 		if (FCoreDelegates::MountPak.Execute(File, INDEX_NONE)) //May fail on encrypted Paks
@@ -2103,7 +2114,8 @@ void FDefaultInstallBundleManager::TickReleaseRequests()
 
 				InstallBundleManagerAnalytics::FireEvent_BundleReleaseRequestStarted(AnalyticsProvider.Get(),
 					Request->BundleName.ToString(),
-					EnumHasAnyFlags(Request->Flags, EInstallBundleReleaseRequestFlags::RemoveFilesIfPossible));
+					EnumHasAnyFlags(Request->Flags, EInstallBundleReleaseRequestFlags::RemoveFilesIfPossible),
+					EnumHasAnyFlags(Request->Flags, EInstallBundleReleaseRequestFlags::SkipReleaseUnmountOnly));
 
 				ContentReleaseRequests[EContentReleaseRequestBatch::Requested].RemoveAtSwap(i);
 				bRequestComplete = true;
@@ -2154,6 +2166,7 @@ void FDefaultInstallBundleManager::TickReleaseRequests()
 				InstallBundleManagerAnalytics::FireEvent_BundleReleaseRequestComplete(AnalyticsProvider.Get(),
 					Request->BundleName.ToString(),
 					EnumHasAnyFlags(Request->Flags, EInstallBundleReleaseRequestFlags::RemoveFilesIfPossible),
+					EnumHasAnyFlags(Request->Flags, EInstallBundleReleaseRequestFlags::SkipReleaseUnmountOnly),
 					LexToString(Request->Result));
 
 				LOG_INSTALL_BUNDLE_MAN_OVERRIDE(Request->LogVerbosityOverride, Display, TEXT("Removing Release Request %s"), *Request->BundleName.ToString());
@@ -3432,6 +3445,12 @@ TValueOrError<FInstallBundleReleaseRequestInfo, EInstallBundleResult> FDefaultIn
 		return MakeError(EInstallBundleResult::InitializationPending);
 	}
 
+	// RemoveFilesIfPossible and SkipReleaseUnmountOnly are incompatible
+	if (EnumHasAllFlags(Flags, EInstallBundleReleaseRequestFlags::RemoveFilesIfPossible | EInstallBundleReleaseRequestFlags::SkipReleaseUnmountOnly))
+	{
+		return MakeError(EInstallBundleResult::InstallError);
+	}
+
 	LOG_INSTALL_BUNDLE_MAN_OVERRIDE(LogVerbosityOverride, Display, TEXT("RequestReleaseContent"));
 
 	TSet<FName> BundlesToKeep = GatherBundlesForRequest(KeepNames, RetInfo.InfoFlags);
@@ -3491,6 +3510,13 @@ TValueOrError<FInstallBundleReleaseRequestInfo, EInstallBundleResult> FDefaultIn
 		// canceled update has finished.
 		if (ActiveQueuedRequest == nullptr && !bCanceledUpdate)
 		{
+			if (EnumHasAnyFlags(Flags, EInstallBundleReleaseRequestFlags::SkipReleaseUnmountOnly) && GetBundleStatus(*BundleInfo) != EBundleState::Mounted)
+			{
+				RetInfo.InfoFlags |= EInstallBundleRequestInfoFlags::SkippedAlreadyReleasedBundles;
+				LOG_INSTALL_BUNDLE_MAN_OVERRIDE(LogVerbosityOverride, Verbose, TEXT("BundlesToRelease Bundle %s  - Already Unmounted"), *BundleInfo->BundleNameString);
+				continue;
+			}
+
 			// If this bundle is reserved in a cache, a release request cannot be skipped
 			bool bIsReserved = false;
 			for (const FBundleSourceRelevance& SourceRelevance : BundleInfo->ContributingSources)
@@ -3527,24 +3553,42 @@ TValueOrError<FInstallBundleReleaseRequestInfo, EInstallBundleResult> FDefaultIn
 
 		LOG_INSTALL_BUNDLE_MAN_OVERRIDE(LogVerbosityOverride, Display, TEXT("Requesting Release of Bundle %s"), *BundleInfo->BundleNameString);
 
-		if (ActiveQueuedRequest &&
-			!EnumHasAnyFlags(ActiveQueuedRequest->Flags, EInstallBundleReleaseRequestFlags::RemoveFilesIfPossible) &&
-			EnumHasAnyFlags(Flags, EInstallBundleReleaseRequestFlags::RemoveFilesIfPossible))
+		if (ActiveQueuedRequest)
 		{
-			ActiveQueuedRequest->Flags |= EInstallBundleReleaseRequestFlags::RemoveFilesIfPossible;
+			if (!EnumHasAnyFlags(ActiveQueuedRequest->Flags, EInstallBundleReleaseRequestFlags::RemoveFilesIfPossible) &&
+				EnumHasAnyFlags(Flags, EInstallBundleReleaseRequestFlags::RemoveFilesIfPossible))
+			{
+				ActiveQueuedRequest->Flags |= EInstallBundleReleaseRequestFlags::RemoveFilesIfPossible;
 
-			// TODO: This assumes that a bundle source that doesn't remove files will always immediatly callback on release.
-			// That is probably a bad assumption, we can't control how the bundle source may be written.
-			// It would be safer to instead in enqueue this with a prereq that there is no pending release
+				// TODO: This assumes that a bundle source that doesn't remove files will always immediatly callback on release.
+				// That is probably a bad assumption, we can't control how the bundle source may be written.
+				// It would be safer to instead in enqueue this with a prereq that there is no pending release
 #if DO_CHECK
-			// Since a request without RemoveFilesIfPossible does no async work after unmounting, it shouldn't be possible to 
-			// call RequestReleaseContent and find an ActiveQueuedRequest that is past the unmounting step.
-			if (ActiveQueuedRequest->Steps.IsValidIndex(ActiveQueuedRequest->iStep))
-			{
-				EContentReleaseRequestState State = ActiveQueuedRequest->Steps[ActiveQueuedRequest->iStep];
-				check(State < EContentReleaseRequestState::UpdatingBundleSources);
+				// Since a request without RemoveFilesIfPossible does no async work after unmounting, it shouldn't be possible to 
+				// call RequestReleaseContent and find an ActiveQueuedRequest that is past the unmounting step.
+				if (ActiveQueuedRequest->Steps.IsValidIndex(ActiveQueuedRequest->iStep))
+				{
+					EContentReleaseRequestState State = ActiveQueuedRequest->Steps[ActiveQueuedRequest->iStep];
+					check(State < EContentReleaseRequestState::UpdatingBundleSources);
+				}
+#endif // DO_CHECK
 			}
+
+			if (EnumHasAnyFlags(ActiveQueuedRequest->Flags, EInstallBundleReleaseRequestFlags::SkipReleaseUnmountOnly) &&
+				!EnumHasAnyFlags(Flags, EInstallBundleReleaseRequestFlags::SkipReleaseUnmountOnly))
+			{
+				ActiveQueuedRequest->Flags &= ~EInstallBundleReleaseRequestFlags::SkipReleaseUnmountOnly;
+
+#if DO_CHECK
+				// Since a request with SkipReleaseUnmountOnly does no async work after unmounting, it shouldn't be possible to 
+				// call RequestReleaseContent and find an ActiveQueuedRequest that is past the unmounting step.
+				if (ActiveQueuedRequest->Steps.IsValidIndex(ActiveQueuedRequest->iStep))
+				{
+					EContentReleaseRequestState State = ActiveQueuedRequest->Steps[ActiveQueuedRequest->iStep];
+					check(State < EContentReleaseRequestState::UpdatingBundleSources);
+				}
 #endif // DO_CHECK
+			}
 		}
 
 		if (ActiveQueuedRequest == nullptr)
diff --git a/Engine/Plugins/Experimental/DefaultInstallBundleManager/Source/Private/InstallBundleManagerUtil.cpp b/Engine/Plugins/Experimental/DefaultInstallBundleManager/Source/Private/InstallBundleManagerUtil.cpp
--- a/Engine/Plugins/Experimental/DefaultInstallBundleManager/Source/Private/InstallBundleManagerUtil.cpp
+++ b/Engine/Plugins/Experimental/DefaultInstallBundleManager/Source/Private/InstallBundleManagerUtil.cpp
@@ -744,7 +744,10 @@ namespace InstallBundleManagerAnalytics
 		AnalyticsProvider->RecordEvent(TEXT("InstallBundleManager.BundleRequestComplete"), MoveTemp(Attributes));
 	}
 
-	void FireEvent_BundleReleaseRequestStarted(IAnalyticsProviderET* AnalyticsProvider, const FString& BundleName, bool bRemoveFilesIfPossible)
+	void FireEvent_BundleReleaseRequestStarted(IAnalyticsProviderET* AnalyticsProvider, 
+		const FString& BundleName, 
+		bool bRemoveFilesIfPossible,
+		bool bUnmountOnly)
 	{
 		if (AnalyticsProvider == nullptr || InstallBundleUtil::FInstallBundleSuppressAnalytics::IsEnabled())
 		{
@@ -753,12 +756,14 @@ namespace InstallBundleManagerAnalytics
 
 		AnalyticsProvider->RecordEvent(TEXT("InstallBundleManager.BundleReleaseRequestStarted"), MakeAnalyticsEventAttributeArray(
 			TEXT("BundleName"), BundleName,
-			TEXT("RemoveFilesIfPossible"), bRemoveFilesIfPossible));
+			TEXT("RemoveFilesIfPossible"), bRemoveFilesIfPossible,
+			TEXT("UnmountOnly"), bUnmountOnly));
 	}
 
 	void FireEvent_BundleReleaseRequestComplete(IAnalyticsProviderET* AnalyticsProvider, 
 		const FString& BundleName, 
 		bool bRemoveFilesIfPossible, 
+		bool bUnmountOnly,
 		const FString& Result)
 	{
 		if (AnalyticsProvider == nullptr || InstallBundleUtil::FInstallBundleSuppressAnalytics::IsEnabled())
@@ -769,6 +774,7 @@ namespace InstallBundleManagerAnalytics
 		AnalyticsProvider->RecordEvent(TEXT("InstallBundleManager.BundleReleaseRequestComplete"), MakeAnalyticsEventAttributeArray(
 			TEXT("BundleName"), BundleName,
 			TEXT("RemoveFilesIfPossible"), bRemoveFilesIfPossible,
+			TEXT("UnmountOnly"), bUnmountOnly,
 			TEXT("Result"), Result));
 	}
 
diff --git a/Engine/Plugins/Experimental/DefaultInstallBundleManager/Source/Public/InstallBundleManagerUtil.h b/Engine/Plugins/Experimental/DefaultInstallBundleManager/Source/Public/InstallBundleManagerUtil.h
--- a/Engine/Plugins/Experimental/DefaultInstallBundleManager/Source/Public/InstallBundleManagerUtil.h
+++ b/Engine/Plugins/Experimental/DefaultInstallBundleManager/Source/Public/InstallBundleManagerUtil.h
@@ -335,10 +335,12 @@ namespace InstallBundleManagerAnalytics
 	 *
 	 * @EventParam BundleName (string) String Name of bundle.
 	 * @EventParam RemoveFilesIfPossible (bool) True if this request will try to clean up files.
+	 * @EventParam UnmountOnly (bool) True if this request is only for unmount
 	 */
 	DEFAULTINSTALLBUNDLEMANAGER_API void FireEvent_BundleReleaseRequestStarted(IAnalyticsProviderET* AnalyticsProvider,
 		const FString& BundleName,
-		bool bRemoveFilesIfPossible);
+		bool bRemoveFilesIfPossible,
+		bool bUnmountOnly);
 
 	/**
 	 * @EventName InstallBundleManager.BundleReleaseRequestComplete
@@ -349,11 +351,13 @@ namespace InstallBundleManagerAnalytics
 	 *
 	 * @EventParam BundleName (string) String Name of bundle.
 	 * @EventParam RemoveFilesIfPossible (bool) True if this request will try to clean up files.
+	 * @EventParam UnmountOnly (bool) True if this request is only for unmount
 	 * @EventParam Result (string) Result of the request
 	 */
 	DEFAULTINSTALLBUNDLEMANAGER_API void FireEvent_BundleReleaseRequestComplete(IAnalyticsProviderET* AnalyticsProvider,
 		const FString& BundleName,
 		bool bRemoveFilesIfPossible,
+		bool bUnmountOnly,
 		const FString& Result);
 
 	/**
diff --git a/Engine/Plugins/Runtime/GameFeatures/Source/GameFeatures/Private/GameFeaturePluginStateMachine.cpp b/Engine/Plugins/Runtime/GameFeatures/Source/GameFeatures/Private/GameFeaturePluginStateMachine.cpp
--- a/Engine/Plugins/Runtime/GameFeatures/Source/GameFeatures/Private/GameFeaturePluginStateMachine.cpp
+++ b/Engine/Plugins/Runtime/GameFeatures/Source/GameFeatures/Private/GameFeaturePluginStateMachine.cpp
@@ -1084,7 +1084,9 @@ struct FBaseDataReleaseGameFeaturePluginState : public FGameFeaturePluginState
 	/** Determine what kind of release request flags we submit */
 	virtual EInstallBundleReleaseRequestFlags GetReleaseRequestFlags() const
 	{
-		return StateProperties.ProtocolOptions.GetSubtype<FInstallBundlePluginProtocolOptions>().ReleaseInstallBundleFlags;
+		EInstallBundleReleaseRequestFlags ReleaseFlags = StateProperties.ProtocolOptions.GetSubtype<FInstallBundlePluginProtocolOptions>().ReleaseInstallBundleFlags;
+		ReleaseFlags &= ~EInstallBundleReleaseRequestFlags::SkipReleaseUnmountOnly; // Release state MUST release
+		return ReleaseFlags;
 	}
 
 	/** Determines what state you transition to in the event of a success or failure to release content */
@@ -1662,13 +1664,12 @@ struct FGameFeaturePluginState_Unmounting : public FGameFeaturePluginState
 			return;
 		}
 
-		UGameFeaturesSubsystem::Get().OnGameFeatureReleasing(StateProperties.PluginName, StateProperties.PluginIdentifier);
-
 		TSharedPtr<IInstallBundleManager> BundleManager = IInstallBundleManager::GetPlatformInstallBundleManager();
 
 		const TArray<FName>& InstallBundles = StateProperties.ProtocolMetadata.GetSubtype<FInstallBundlePluginProtocolMetaData>().InstallBundles;
 
 		EInstallBundleReleaseRequestFlags ReleaseFlags = StateProperties.ProtocolOptions.GetSubtype<FInstallBundlePluginProtocolOptions>().ReleaseInstallBundleFlags;
+		ReleaseFlags |= EInstallBundleReleaseRequestFlags::SkipReleaseUnmountOnly;
 		//Make sure we don't remove files here early, that should only be done in Uninstalling
 		ReleaseFlags &= ~(EInstallBundleReleaseRequestFlags::RemoveFilesIfPossible);
 
diff --git a/Engine/Plugins/Runtime/GameFeatures/Source/GameFeatures/Public/GameFeaturesSubsystem.h b/Engine/Plugins/Runtime/GameFeatures/Source/GameFeatures/Public/GameFeaturesSubsystem.h
--- a/Engine/Plugins/Runtime/GameFeatures/Source/GameFeatures/Public/GameFeaturesSubsystem.h
+++ b/Engine/Plugins/Runtime/GameFeatures/Source/GameFeatures/Public/GameFeaturesSubsystem.h
@@ -652,7 +652,6 @@ private:
 
 	void OnGameFeatureReleasing(const FString& PluginName, const FGameFeaturePluginIdentifier& PluginIdentifier);
 	friend struct FGameFeaturePluginState_Releasing;
-	friend struct FGameFeaturePluginState_Unmounting;
 
 	void OnGameFeaturePreMounting(const FString& PluginName, const FGameFeaturePluginIdentifier& PluginIdentifier, FGameFeaturePreMountingContext& Context);
 	void OnGameFeaturePostMounting(const FString& PluginName, const FGameFeaturePluginIdentifier& PluginIdentifier, FGameFeaturePostMountingContext& Context);
diff --git a/Engine/Source/Runtime/InstallBundleManager/Public/InstallBundleTypes.h b/Engine/Source/Runtime/InstallBundleManager/Public/InstallBundleTypes.h
--- a/Engine/Source/Runtime/InstallBundleManager/Public/InstallBundleTypes.h
+++ b/Engine/Source/Runtime/InstallBundleManager/Public/InstallBundleTypes.h
@@ -208,7 +208,8 @@ enum class EInstallBundleReleaseRequestFlags : uint32
 {
 	None = 0,
 	RemoveFilesIfPossible = (1 << 0),  // Bundle sources must support removal, and bundle must not be part of the source's cache
-	ExplicitRemoveList = (1 << 1),	   // Only attempt to remove explicitly supplied bundles instead of automatically removing dependencies 
+	ExplicitRemoveList = (1 << 1),	   // Only attempt to remove explicitly supplied bundles instead of automatically removing dependencies
+	SkipReleaseUnmountOnly = (1 << 2),   // Unmount but leave content referenced. The inverse of EInstallBundleRequestFlags::SkipMount
 };
 ENUM_CLASS_FLAGS(EInstallBundleReleaseRequestFlags)
 
