Platform chunk installers cannot be used by the bundle source wrapper if chunk installation is not available (e.g. if the build is not packaged)

#rnx

diff --git a/Engine/Plugins/Experimental/DefaultInstallBundleManager/Source/Private/InstallBundleSourcePlatformChunkInstall.cpp b/Engine/Plugins/Experimental/DefaultInstallBundleManager/Source/Private/InstallBundleSourcePlatformChunkInstall.cpp
--- a/Engine/Plugins/Experimental/DefaultInstallBundleManager/Source/Private/InstallBundleSourcePlatformChunkInstall.cpp
+++ b/Engine/Plugins/Experimental/DefaultInstallBundleManager/Source/Private/InstallBundleSourcePlatformChunkInstall.cpp
@@ -125,9 +125,9 @@ FInstallBundleSourceInitInfo FInstallBundleSourcePlatformChunkInstall::Init(TSha
 	CSV_SCOPED_TIMING_STAT(InstallBundleManager, FInstallBundleSourcePlatformChunkInstall_Init);
 	FInstallBundleSourceInitInfo InitInfo;
 
-	if (!PlatformChunkInstall->SupportsBundleSource())
+	if (!PlatformChunkInstall->SupportsBundleSource() || !PlatformChunkInstall->IsAvailable())
 	{
-		LOG_SOURCE_CHUNKINSTALL(Display, TEXT("Platform chunk installer doesn't support bundles, attempting to fallback to next bundle source."));
+		LOG_SOURCE_CHUNKINSTALL(Display, TEXT("Platform chunk installer doesn't support bundles or this is not a packaged build, attempting to fallback to next bundle source."));
 
 		InitState = EInstallBundleManagerInitState::Failed;
 
diff --git a/Engine/Source/Runtime/Core/Public/GenericPlatform/GenericPlatformChunkInstall.h b/Engine/Source/Runtime/Core/Public/GenericPlatform/GenericPlatformChunkInstall.h
--- a/Engine/Source/Runtime/Core/Public/GenericPlatform/GenericPlatformChunkInstall.h
+++ b/Engine/Source/Runtime/Core/Public/GenericPlatform/GenericPlatformChunkInstall.h
@@ -152,6 +152,11 @@ public:
 	/** Virtual destructor */
 	virtual ~IPlatformChunkInstall() {}
 
+	/** 
+	 * Returns whether chunk installation is available (i.e. we are an installed packaged build etc)
+	 */
+	virtual bool IsAvailable() const = 0;
+
 	/**
 	 * Get the current location of a chunk with pakchunk index.
 	 * @param PakchunkIndex	The id of the pak chunk.
@@ -394,6 +399,11 @@ PRAGMA_DISABLE_DEPRECATION_WARNINGS
 class CORE_API FGenericPlatformChunkInstall : public IPlatformChunkInstall
 {
 public:
+	virtual bool IsAvailable() const
+	{
+		return true;
+	}
+
 	virtual EChunkLocation::Type GetPakchunkLocation( int32 PakchunkIndex ) override
 	{
 		return GetChunkLocation(PakchunkIndex);
