UpdateBundleSourceReleaseComplete is more robust when dealing with partial content removal from a source.

diff --git a/Engine/Plugins/Experimental/DefaultInstallBundleManager/Source/Private/DefaultInstallBundleManager.cpp b/Engine/Plugins/Experimental/DefaultInstallBundleManager/Source/Private/DefaultInstallBundleManager.cpp
--- a/Engine/Plugins/Experimental/DefaultInstallBundleManager/Source/Private/DefaultInstallBundleManager.cpp
+++ b/Engine/Plugins/Experimental/DefaultInstallBundleManager/Source/Private/DefaultInstallBundleManager.cpp
@@ -1445,10 +1445,9 @@ void FDefaultInstallBundleManager::UpdateBundleSourceReleaseComplete(TSharedRef<
 		{
 			// Removing content wasn't requested, but the source did it anyway so deal with it
 			bContentWasRemoved = true;
-			continue;
 		}
 
-		// Since we didn't remove content, update last access times in any caches this bundle participates in
+		// Update last access times in any caches this bundle participates in
 		if (const FName* CacheName = BundleSourceCaches.Find(Pair.Key))
 		{
 			const TSharedRef<FInstallBundleCache>& BundleCache = BundleCaches.FindChecked(*CacheName);
@@ -1464,9 +1463,9 @@ void FDefaultInstallBundleManager::UpdateBundleSourceReleaseComplete(TSharedRef<
 
 	LOG_INSTALL_BUNDLE_MAN_OVERRIDE(Request->LogVerbosityOverride, Display, TEXT("Release of Bundle %s done waiting for all bundle sources!"), *BundleInfo.BundleNameString);
 
-	if (Request->Result == EInstallBundleReleaseResult::OK && bContentWasRemoved)
+	if (bContentWasRemoved)
 	{
-		// The Bundle is no longer installed because at least one of its sources is completely uninstalled although
+		// The Bundle is no longer installed because at least one of its sources is not completely installed although
 		// it may still have data in cached sources
 		SetBundleStatus(BundleInfo, EBundleState::NotInstalled);
 	}
