Fix call to release content on a install bundle source not occuring if the bundle was unmounted but not fully released.
Allow bundle sources to add additional bundles to be canceled when they get a cancel request.  This is useful if the source has internal dependencies between bundles.

diff --git a/Engine/Plugins/Experimental/DefaultInstallBundleManager/Source/Private/DefaultInstallBundleManager.cpp b/Engine/Plugins/Experimental/DefaultInstallBundleManager/Source/Private/DefaultInstallBundleManager.cpp
--- a/Engine/Plugins/Experimental/DefaultInstallBundleManager/Source/Private/DefaultInstallBundleManager.cpp
+++ b/Engine/Plugins/Experimental/DefaultInstallBundleManager/Source/Private/DefaultInstallBundleManager.cpp
@@ -847,7 +847,7 @@ void FDefaultInstallBundleManager::DetermineSteps(FContentRequestRef Request)
 
 void FDefaultInstallBundleManager::AddRequestToInitialBatch(FContentRequestRef Request)
 {
-	const FBundleInfo& BundleInfo = BundleInfoMap.FindChecked(Request->BundleName);
+	FBundleInfo& BundleInfo = BundleInfoMap.FindChecked(Request->BundleName);
 
 	EContentRequestBatch InitialBatch = EContentRequestBatch::Cache;
 	switch (GetBundleStatus(BundleInfo))
@@ -858,6 +858,7 @@ void FDefaultInstallBundleManager::AddRequestToInitialBatch(FContentRequestRef R
 
 	case EBundleState::Mounted:
 		InitialBatch = EContentRequestBatch::Install;
+		BundleInfo.bReleaseRequired = true;
 		Request->bShouldSendAnalytics = false;
 		break;
 	}
@@ -1954,6 +1955,9 @@ void FDefaultInstallBundleManager::TickContentRequests()
 
 		Request->StepResult = EContentRequestStepResult::Done;
 
+		FBundleInfo& BundleInfo = BundleInfoMap.FindChecked(Request->BundleName);
+		BundleInfo.bReleaseRequired = true;
+
 		LOG_INSTALL_BUNDLE_MAN_OVERRIDE(Request->LogVerbosityOverride, Verbose, TEXT("Moving Request %s from batch %s to batch %s"), *Request->BundleName.ToString(), LexToString(EContentRequestBatch::Cache), LexToString(EContentRequestBatch::Install));
 		ContentRequests[EContentRequestBatch::Install].Add(MoveTemp(Request));
 		ContentRequests[EContentRequestBatch::Cache].RemoveAt(ContentRequests[EContentRequestBatch::Cache].Num() - 1); //Invalidates Request
@@ -3130,7 +3134,7 @@ TValueOrError<FInstallBundleRequestInfo, EInstallBundleResult> FDefaultInstallBu
 			// canceled release has finished.
 			if (!bCanceledRelease && EnumHasAnyFlags(Flags, EInstallBundleRequestFlags::SkipMount) && GetBundleStatus(*BundleInfo) == EBundleState::NeedsMount)
 			{
-				// If this bundle is not reserved in a cache, an  install request cannot be skipped
+				// If this bundle is not reserved in a cache, an install request cannot be skipped
 				bool bNeedsCacheReserve = false;
 				for (const FBundleSourceRelevance& SourceRelevance : BundleInfo->ContributingSources)
 				{
@@ -3516,29 +3520,32 @@ TValueOrError<FInstallBundleReleaseRequestInfo, EInstallBundleResult> FDefaultIn
 				continue;
 			}
 
-			// If this bundle is reserved in a cache, a release request cannot be skipped
-			bool bIsReserved = false;
-			for (const FBundleSourceRelevance& SourceRelevance : BundleInfo->ContributingSources)
+			bool bCanSkipRelease = !BundleInfo->bReleaseRequired;
+			if (bCanSkipRelease)
 			{
-				if (FName* CacheName = BundleSourceCaches.Find(SourceRelevance.SourceType))
+				// If this bundle is reserved in a cache, a release request cannot be skipped
+				for (const FBundleSourceRelevance& SourceRelevance : BundleInfo->ContributingSources)
 				{
-					const TSharedRef<FInstallBundleCache>& BundleCache = BundleCaches.FindChecked(*CacheName);
-					if (BundleCache->IsReserved(BundleName))
+					if (FName* CacheName = BundleSourceCaches.Find(SourceRelevance.SourceType))
 					{
-						bIsReserved = true;
-						break;
+						const TSharedRef<FInstallBundleCache>& BundleCache = BundleCaches.FindChecked(*CacheName);
+						if (BundleCache->IsReserved(BundleName))
+						{
+							bCanSkipRelease = false;
+							break;
+						}
 					}
 				}
 			}
 
-			if (!bIsReserved && !EnumHasAnyFlags(Flags, EInstallBundleReleaseRequestFlags::RemoveFilesIfPossible) && GetBundleStatus(*BundleInfo) != EBundleState::Mounted)
+			if (bCanSkipRelease && !EnumHasAnyFlags(Flags, EInstallBundleReleaseRequestFlags::RemoveFilesIfPossible) && GetBundleStatus(*BundleInfo) != EBundleState::Mounted)
 			{
 				RetInfo.InfoFlags |= EInstallBundleRequestInfoFlags::SkippedAlreadyReleasedBundles;
 				LOG_INSTALL_BUNDLE_MAN_OVERRIDE(LogVerbosityOverride, Verbose, TEXT("BundlesToRelease Bundle %s  - Already Released"), *BundleInfo->BundleNameString);
 				continue;
 			}
 
-			if (!bIsReserved && GetBundleStatus(*BundleInfo) == EBundleState::NotInstalled)
+			if (bCanSkipRelease && GetBundleStatus(*BundleInfo) == EBundleState::NotInstalled)
 			{
 				RetInfo.InfoFlags |= EInstallBundleRequestInfoFlags::SkippedAlreadyRemovedBundles;
 				LOG_INSTALL_BUNDLE_MAN_OVERRIDE(LogVerbosityOverride, Verbose, TEXT("BundlesToRelease Bundle %s  - Already Removed"), *BundleInfo->BundleNameString);
@@ -3725,47 +3732,58 @@ void FDefaultInstallBundleManager::CancelUpdateContent(TArrayView<const FName> B
 	CancelUpdateContentInternal(BundleNames);
 }
 
-bool FDefaultInstallBundleManager::CancelUpdateContentInternal(TArrayView<const FName> BundleNames)
+bool FDefaultInstallBundleManager::CancelUpdateContentInternal(TArrayView<const FName> InBundleNames)
 {
 	bool bCancledRequest = false;
+	TConstArrayView<FName> BundleNames = InBundleNames;
+	TArray<FName> CurrBundleNames;
 
-	for (EContentRequestBatch b : TEnumRange<EContentRequestBatch>())
+	while (BundleNames.Num() > 0)
 	{
-		for (FContentRequestRef& Request : ContentRequests[b])
+		TSet<FName> BundleNamesNext;
+		for (EContentRequestBatch b : TEnumRange<EContentRequestBatch>())
 		{
-			if (!BundleNames.Contains(Request->BundleName))
-				continue;
+			for (FContentRequestRef& Request : ContentRequests[b])
+			{
+				if (!BundleNames.Contains(Request->BundleName))
+					continue;
 
-			bCancledRequest = true;
+				bCancledRequest = true;
 
-			if (Request->bIsCanceled)
-			{
-				// Already canceled
-				if (Request->Result != EInstallBundleResult::UserCancelledError) // User cancel always has priority
+				if (Request->bIsCanceled)
 				{
-					Request->Result = EInstallBundleResult::UserCancelledError;
+					// Already canceled
+					if (Request->Result != EInstallBundleResult::UserCancelledError) // User cancel always has priority
+					{
+						Request->Result = EInstallBundleResult::UserCancelledError;
+					}
+					continue;
 				}
-				continue;
-			}
 
-			if (Request->Steps.IsValidIndex(Request->iStep))
-			{
-				EContentRequestState State = Request->Steps[Request->iStep];
-				if (State == EContentRequestState::UpdatingBundleSources)
+				if (Request->Steps.IsValidIndex(Request->iStep))
 				{
-					for (const TPair<FInstallBundleSourceType, TSharedPtr<IInstallBundleSource>>& Pair : BundleSources)
+					EContentRequestState State = Request->Steps[Request->iStep];
+					if (State == EContentRequestState::UpdatingBundleSources)
 					{
-						Pair.Value->CancelBundles(BundleNames);
+						for (const TPair<FInstallBundleSourceType, TSharedPtr<IInstallBundleSource>>& Pair : BundleSources)
+						{
+							TArray<FName> AdditionalBundles;
+							Pair.Value->CancelBundles(BundleNames, AdditionalBundles);
+							BundleNamesNext.Append(AdditionalBundles);
+						}
 					}
 				}
-			}
 
-			LOG_INSTALL_BUNDLE_MAN_OVERRIDE(Request->LogVerbosityOverride, Display, TEXT("Canceling Install Request %s Result: %s"),
-				*Request->BundleName.ToString(), LexToString(EInstallBundleResult::UserCancelledError));
+				LOG_INSTALL_BUNDLE_MAN_OVERRIDE(Request->LogVerbosityOverride, Display, TEXT("Canceling Install Request %s Result: %s"),
+					*Request->BundleName.ToString(), LexToString(EInstallBundleResult::UserCancelledError));
 
-			Request->bIsCanceled = true;
-			Request->Result = EInstallBundleResult::UserCancelledError;
+				Request->bIsCanceled = true;
+				Request->Result = EInstallBundleResult::UserCancelledError;
+			}
 		}
+
+		CurrBundleNames = BundleNamesNext.Array();
+		BundleNames = CurrBundleNames;
 	}
 
 	return bCancledRequest;
diff --git a/Engine/Plugins/Experimental/DefaultInstallBundleManager/Source/Public/DefaultInstallBundleManager.h b/Engine/Plugins/Experimental/DefaultInstallBundleManager/Source/Public/DefaultInstallBundleManager.h
--- a/Engine/Plugins/Experimental/DefaultInstallBundleManager/Source/Public/DefaultInstallBundleManager.h
+++ b/Engine/Plugins/Experimental/DefaultInstallBundleManager/Source/Public/DefaultInstallBundleManager.h
@@ -196,6 +196,7 @@ protected:
 	public:
 		FString BundleNameString; // Since FNames do not preserve casing
 		TArray<EBundlePrereqs> Prereqs;
+		bool bReleaseRequired = false; // A bundle source may have done install work so expects a release call
 		bool bIsStartup = false;
 		TArray<FBundleSourceRelevance> ContributingSources; // Sources contributing to this bundle info
 		EInstallBundlePriority Priority = EInstallBundlePriority::Low;
diff --git a/Engine/Source/Runtime/InstallBundleManager/Public/InstallBundleSourceInterface.h b/Engine/Source/Runtime/InstallBundleManager/Public/InstallBundleSourceInterface.h
--- a/Engine/Source/Runtime/InstallBundleManager/Public/InstallBundleSourceInterface.h
+++ b/Engine/Source/Runtime/InstallBundleManager/Public/InstallBundleSourceInterface.h
@@ -118,6 +118,7 @@ public:
 
 	// Cancel the install for the specified bundles
 	virtual void CancelBundles(TArrayView<const FName> BundleNames) {}
+	virtual void CancelBundles(TArrayView<const FName> BundleNames, TArray<FName>& OutAdditionalBunldlesToCancel) { CancelBundles(BundleNames); }
 
 	// User Pause/Resume bundles.
 	virtual void UserPauseBundles(TArrayView<const FName> BundleNames) {}
