Use the new "MaxDiskSpaceRequiredIfDeleteDuring" bundle ini setting for installs requiring less free drive space.

#rnx

diff --git a/Engine/Plugins/Experimental/DefaultInstallBundleManager/Source/Private/DefaultInstallBundleManager.cpp b/Engine/Plugins/Experimental/DefaultInstallBundleManager/Source/Private/DefaultInstallBundleManager.cpp
--- a/Engine/Plugins/Experimental/DefaultInstallBundleManager/Source/Private/DefaultInstallBundleManager.cpp
+++ b/Engine/Plugins/Experimental/DefaultInstallBundleManager/Source/Private/DefaultInstallBundleManager.cpp
@@ -513,6 +513,7 @@ void FDefaultInstallBundleManager::TickGetContentState()
 				State.DownloadSize += SourceState.DownloadSize;
 				State.InstallSize += SourceState.InstallSize;
 				State.InstallOverheadSize = FMath::Max(State.InstallOverheadSize, SourceState.InstallOverheadSize);
+				State.MaxDiskSpaceRequired += SourceState.MaxDiskSpaceRequired;
 				State.FreeSpace = FMath::Max(State.FreeSpace, SourceState.FreeSpace);
 			}
 
diff --git a/Engine/Source/Runtime/InstallBundleManager/Public/InstallBundleTypes.h b/Engine/Source/Runtime/InstallBundleManager/Public/InstallBundleTypes.h
--- a/Engine/Source/Runtime/InstallBundleManager/Public/InstallBundleTypes.h
+++ b/Engine/Source/Runtime/InstallBundleManager/Public/InstallBundleTypes.h
@@ -131,6 +131,7 @@ struct FInstallBundleCombinedContentState
 	uint64 DownloadSize = 0;
 	uint64 InstallSize = 0;
 	uint64 InstallOverheadSize = 0;
+	uint64 MaxDiskSpaceRequired = 0;
 	uint64 FreeSpace = 0;
 
 	INSTALLBUNDLEMANAGER_API bool GetAllBundlesHaveState(EInstallBundleInstallState State, TArrayView<const FName> ExcludedBundles = TArrayView<const FName>()) const;	
