Remove use of OSID and LoginID from FDefaultInstallBundleManager::AsyncInit_CreateAnalyticsSession()

diff --git a/Engine/Plugins/Experimental/DefaultInstallBundleManager/Source/Private/DefaultInstallBundleManager.cpp b/Engine/Plugins/Experimental/DefaultInstallBundleManager/Source/Private/DefaultInstallBundleManager.cpp
--- a/Engine/Plugins/Experimental/DefaultInstallBundleManager/Source/Private/DefaultInstallBundleManager.cpp
+++ b/Engine/Plugins/Experimental/DefaultInstallBundleManager/Source/Private/DefaultInstallBundleManager.cpp
@@ -4792,13 +4792,6 @@ void FDefaultInstallBundleManager::AsyncInit_CreateAnalyticsSession()
 
 	if (AnalyticsProvider)
 	{
-		FString MachineID = FPlatformMisc::GetOperatingSystemId();
-		if (MachineID.IsEmpty())
-		{
-			MachineID = FPlatformMisc::GetLoginId();
-		}
-		ensureAlways(!MachineID.IsEmpty());
-
 		// Pick the latest content version from all sources
 		FString ContentVersion;
 		int64 MaxVersionCL = -1;
@@ -4823,7 +4816,7 @@ void FDefaultInstallBundleManager::AsyncInit_CreateAnalyticsSession()
 		}
 		ensureAlways(MaxVersionCL != -1);
 
-		AnalyticsProvider->SetSessionID(FString(TEXT("InstallSession-")) + MachineID + TEXT("-") + ContentVersion);
+		AnalyticsProvider->SetSessionID(FString(TEXT("IBMInstallSession-")) + FGuid().ToString() + TEXT("-") + ContentVersion);
 	}
 	InitStepResult = EAsyncInitStepResult::Done;
 }
