Fix AsyncInit_CreateAnalyticsSession SessionID to use FGuid::NewGuid() instead of creating a zero-guid using the default ctor.
#tests compiled
#rnx

diff --git a/Engine/Plugins/Experimental/DefaultInstallBundleManager/Source/Private/DefaultInstallBundleManager.cpp b/Engine/Plugins/Experimental/DefaultInstallBundleManager/Source/Private/DefaultInstallBundleManager.cpp
--- a/Engine/Plugins/Experimental/DefaultInstallBundleManager/Source/Private/DefaultInstallBundleManager.cpp
+++ b/Engine/Plugins/Experimental/DefaultInstallBundleManager/Source/Private/DefaultInstallBundleManager.cpp
@@ -4816,7 +4816,7 @@ void FDefaultInstallBundleManager::AsyncInit_CreateAnalyticsSession()
 		}
 		ensureAlways(MaxVersionCL != -1);
 
-		AnalyticsProvider->SetSessionID(FString(TEXT("IBMInstallSession-")) + FGuid().ToString() + TEXT("-") + ContentVersion);
+		AnalyticsProvider->SetSessionID(FString(TEXT("IBMInstallSession-")) + FGuid::NewGuid().ToString() + TEXT("-") + ContentVersion);
 	}
 	InitStepResult = EAsyncInitStepResult::Done;
 }
