Opt-in async chunk installer initialize for platforms that support it.

 - adding additional logging for chunk installer & bundle wrapper

diff --git a/Engine/Plugins/Experimental/DefaultInstallBundleManager/Source/Private/InstallBundleSourcePlatformChunkInstall.cpp b/Engine/Plugins/Experimental/DefaultInstallBundleManager/Source/Private/InstallBundleSourcePlatformChunkInstall.cpp
--- a/Engine/Plugins/Experimental/DefaultInstallBundleManager/Source/Private/InstallBundleSourcePlatformChunkInstall.cpp
+++ b/Engine/Plugins/Experimental/DefaultInstallBundleManager/Source/Private/InstallBundleSourcePlatformChunkInstall.cpp
@@ -187,6 +187,7 @@ void FInstallBundleSourcePlatformChunkInstall::AsyncInit(FInstallBundleSourceIni
 		// cache all named chunks
 		Context->NamedChunks.Append( Context->PlatformChunkInstall->GetNamedChunksByType(ENamedChunkType::OnDemand) );
 		Context->NamedChunks.Append( Context->PlatformChunkInstall->GetNamedChunksByType(ENamedChunkType::Language) );
+		LOG_SOURCE_CHUNKINSTALL(Display, TEXT("Platform chunk install has %d named chunks"), Context->NamedChunks.Num() );
 
 		// create bundles & chunk data
 		TSet<FName> KnownNamedChunks;
@@ -255,6 +256,8 @@ void FInstallBundleSourcePlatformChunkInstall::AsyncInit(FInstallBundleSourceIni
 					}
 				}
 			}
+
+			LOG_SOURCE_CHUNKINSTALL(Display, TEXT("Initialized chunk %s (%d) for bundle %s %s"), *PlatformChunkName, PlatformChunkID, *BundleName.ToString(), (Context->NamedChunks.Contains(BundleInfo.NamedChunk)) ? TEXT("") : TEXT("...not named chunk!") );
 		}
 
 		// collect pak search directories
@@ -314,6 +317,11 @@ void FInstallBundleSourcePlatformChunkInstall::AsyncInit(FInstallBundleSourceIni
 			BundleInfoMap = MoveTemp(Context->BundleInfoMap);
 			NamedChunks = MoveTemp(Context->NamedChunks);
 
+			for (const auto& BundleInfo : BundleInfoMap)
+			{
+				LOG_SOURCE_CHUNKINSTALL(Display, TEXT("%32s -> %-32s (%d file paths)"), *BundleInfo.Key.ToString(), *BundleInfo.Value.NamedChunk.ToString(), BundleInfo.Value.FilePaths.Num());
+			}
+			LOG_SOURCE_CHUNKINSTALL(Display, TEXT("Initialization completed - %d bundles, %d named chunks"), BundleInfoMap.Num(), NamedChunks.Num() );
 			InitState = EInstallBundleManagerInitState::Succeeded;
 		}
 		else
@@ -328,8 +336,26 @@ void FInstallBundleSourcePlatformChunkInstall::AsyncInit(FInstallBundleSourceIni
 		OnInitCompleteCallback.Execute(AsShared(), MoveTemp(InitInfo));
 	};
 
-	// start the async initialize
-	InstallBundleUtil::StartInstallBundleAsyncIOTask(GeneralAsyncTasks, MoveTemp(AsyncInit), MoveTemp(OnAsyncInitComplete));
+	auto OnPlatformChunkInstallComplete = [this, Context, AsyncInit = MoveTemp(AsyncInit), OnAsyncInitComplete = MoveTemp(OnAsyncInitComplete)](bool bSuccess) mutable
+	{
+		check(IsInGameThread());
+
+		if (bSuccess)
+		{
+			// next start the async bundle initialize
+			LOG_SOURCE_CHUNKINSTALL(Display, TEXT("Platform chunk installer init succeeded") );
+			InstallBundleUtil::StartInstallBundleAsyncIOTask(GeneralAsyncTasks, MoveTemp(AsyncInit), MoveTemp(OnAsyncInitComplete));
+		}
+		else
+		{
+			LOG_SOURCE_CHUNKINSTALL(Error, TEXT("Initialization Failed - platform chunk installer async init failed"));
+			Context->InitResult = EInstallBundleManagerInitResult::ConfigurationError;
+			OnAsyncInitComplete();
+		}
+	};
+	
+	// firstly, start initializing the platform chunk installer
+	PlatformChunkInstall->AsyncInit(OnPlatformChunkInstallComplete);
 }
 
 
diff --git a/Engine/Source/Runtime/Core/Private/GenericPlatform/GenericPlatformChunkInstall.cpp b/Engine/Source/Runtime/Core/Private/GenericPlatform/GenericPlatformChunkInstall.cpp
--- a/Engine/Source/Runtime/Core/Private/GenericPlatform/GenericPlatformChunkInstall.cpp
+++ b/Engine/Source/Runtime/Core/Private/GenericPlatform/GenericPlatformChunkInstall.cpp
@@ -5,6 +5,16 @@
 
 DEFINE_LOG_CATEGORY(LogChunkInstaller)
 
+void FGenericPlatformChunkInstall::AsyncInit( TFunction<void(bool/*bSuccess*/)> OnInitComplete )
+{
+	// default implementation does nothing except trigger the complete callback on the gamethread
+	ExecuteOnGameThread(UE_SOURCE_LOCATION, [OnInitComplete = MoveTemp(OnInitComplete)]()
+	{
+		OnInitComplete(true);
+	});
+}
+
+
 void FGenericPlatformChunkInstall::DoNamedChunkCompleteCallbacks( const FName NamedChunk, EChunkLocation::Type Location, bool bHasSucceeded ) const
 {
 	DoNamedChunkCompleteCallbacks( MakeArrayView(&NamedChunk,1), Location, bHasSucceeded );
diff --git a/Engine/Source/Runtime/Core/Public/GenericPlatform/GenericPlatformChunkInstall.h b/Engine/Source/Runtime/Core/Public/GenericPlatform/GenericPlatformChunkInstall.h
--- a/Engine/Source/Runtime/Core/Public/GenericPlatform/GenericPlatformChunkInstall.h
+++ b/Engine/Source/Runtime/Core/Public/GenericPlatform/GenericPlatformChunkInstall.h
@@ -10,6 +10,7 @@
 #include "Containers/Array.h"
 #include "Containers/ArrayView.h"
 #include "Containers/UnrealString.h"
+#include "Templates/Function.h"
 #include "CoreTypes.h"
 #include "Delegates/Delegate.h"
 #include "HAL/PlatformCrt.h"
@@ -135,6 +136,12 @@ public:
 	/** Virtual destructor */
 	virtual ~IPlatformChunkInstall() {}
 
+	/**
+	 * Initialize the platform chunk installer in the background. May require opt-in for the platform-specific chunk installer due to backwards compatility concerns.
+	 * @param OnInitialized Callback for when asyncronous initialization is completed.
+	 */
+	virtual void AsyncInit( TFunction<void(bool/*bSuccess*/)> OnInitComplete ) = 0;
+
 	/** 
 	 * Returns whether chunk installation is available (i.e. we are an installed packaged build etc)
 	 */
@@ -382,6 +389,8 @@ PRAGMA_DISABLE_DEPRECATION_WARNINGS
 class CORE_API FGenericPlatformChunkInstall : public IPlatformChunkInstall
 {
 public:
+	virtual void AsyncInit( TFunction<void(bool/*bSuccess*/)> OnInitComplete ) override;
+
 	virtual bool IsAvailable() const
 	{
 		return true;
