General IAD Support - Install Bundle changes
- Detect and mount on demand tocs - this always happens, skipmount only applies to paks
- Add flags to indicate on demand tocs are present
- Return the set of bundles with on demand tocs from GetInstallState/GetContentState

diff --git a/Engine/Plugins/Experimental/DefaultInstallBundleManager/Source/DefaultInstallBundleManager.Build.cs b/Engine/Plugins/Experimental/DefaultInstallBundleManager/Source/DefaultInstallBundleManager.Build.cs
--- a/Engine/Plugins/Experimental/DefaultInstallBundleManager/Source/DefaultInstallBundleManager.Build.cs
+++ b/Engine/Plugins/Experimental/DefaultInstallBundleManager/Source/DefaultInstallBundleManager.Build.cs
@@ -23,6 +23,7 @@ public class DefaultInstallBundleManager : ModuleRules
 				"Json",
 				"PakFile",
 				"RenderCore",
+				"IoStoreOnDemand",
 			}
 		);
 	}
diff --git a/Engine/Plugins/Experimental/DefaultInstallBundleManager/Source/Private/DefaultInstallBundleManager.cpp b/Engine/Plugins/Experimental/DefaultInstallBundleManager/Source/Private/DefaultInstallBundleManager.cpp
--- a/Engine/Plugins/Experimental/DefaultInstallBundleManager/Source/Private/DefaultInstallBundleManager.cpp
+++ b/Engine/Plugins/Experimental/DefaultInstallBundleManager/Source/Private/DefaultInstallBundleManager.cpp
@@ -11,6 +11,7 @@
 #include "Misc/Base64.h"
 #include "Misc/ConfigContext.h"
 #include "Misc/CommandLine.h"
+#include "IO/IoStoreOnDemand.h"
 #include "IPlatformFilePak.h"
 #include "HAL/PlatformFileManager.h"
 #include "PatchCheck.h"
@@ -517,6 +518,15 @@ void FDefaultInstallBundleManager::TickGetContentState()
 				State.FreeSpace = State.FreeSpace ? FMath::Min(State.FreeSpace, SourceState.FreeSpace) : SourceState.FreeSpace;
 			}
 
+			for (const TPair<FName, FInstallBundleContentState>& Pair : State.IndividualBundleStates)
+			{
+				const FBundleInfo* BundleInfo = BundleInfoMap.Find(Pair.Key);
+				if (ensure(BundleInfo) && BundleInfo->bContainsIoStoreOnDemandTocs)
+				{
+					State.BundlesWithIoStoreOnDemand.Add(Pair.Key);
+				}
+			}
+
 			Request->ExecCallbackIfValid(MoveTemp(State));
 			GetContentStateRequests.RemoveAt(iRequest);
 			continue;
@@ -573,7 +583,8 @@ FInstallBundleCombinedInstallState FDefaultInstallBundleManager::GetInstallState
 	for (const FName& BundleName : BundleNames)
 	{
 		EInstallBundleInstallState InstallState = EInstallBundleInstallState::NotInstalled;
-		switch (GetBundleStatus(BundleInfoMap.FindChecked(BundleName)))
+		const FBundleInfo& BundleInfo = BundleInfoMap.FindChecked(BundleName);
+		switch (GetBundleStatus(BundleInfo))
 		{
 		case EBundleState::NotInstalled:
 			InstallState = EInstallBundleInstallState::NotInstalled;
@@ -591,6 +602,11 @@ FInstallBundleCombinedInstallState FDefaultInstallBundleManager::GetInstallState
 		}
 
 		RetVal.IndividualBundleStates.Add(BundleName, InstallState);
+
+		if (BundleInfo.bContainsIoStoreOnDemandTocs)
+		{
+			RetVal.BundlesWithIoStoreOnDemand.Add(BundleName);
+		}
 	}
 
 	static_assert(InstallBundleUtil::CastToUnderlying(EBundleState::Count) == 4, "");
@@ -1265,7 +1281,18 @@ void FDefaultInstallBundleManager::UpdateBundleSourceComplete(TSharedRef<IInstal
 		}
 		else if (Request->Result == EInstallBundleResult::OK)
 		{
-			BundleInfo.ContentPaths.ContentPaths.Append(MoveTemp(ResultInfo.ContentPaths));
+			for (FString& ContentPath : ResultInfo.ContentPaths)
+			{
+				// We only need to keep track of the pak paths for unmounting so 
+				// lets save a little memory
+				if (ContentPath.EndsWith(TEXTVIEW(".pak")))
+				{
+					BundleInfo.ContentPaths.ContentPaths.Add(MoveTemp(ContentPath));
+				}
+			}
+			
+			Request->OnDemandMountArgs.Append(MoveTemp(ResultInfo.OnDemandMountArgs));
+
 			BundleInfo.ContentPaths.AdditionalRootDirs.Append(MoveTemp(ResultInfo.AdditionalRootDirs));
 			BundleInfo.ContentPaths.NonUFSShaderLibPaths.Append(MoveTemp(ResultInfo.NonUFSShaderLibPaths));
 			BundleInfo.ContentPaths.ProjectName = ResultInfo.ProjectName;
@@ -1278,6 +1305,9 @@ void FDefaultInstallBundleManager::UpdateBundleSourceComplete(TSharedRef<IInstal
 	}
 	Request->SourceRequestResults.Empty();
 
+	ensureMsgf(!Request->OnDemandMountArgs.IsEmpty() == BundleInfo.bContainsIoStoreOnDemandTocs,
+		TEXT("OnDemandMountArgs should only be provided by bundle sources that set bContainsIoStoreOnDemandTocs"));
+
 	// If there are no content paths, its likely this is a chunk that doesn't exist on the current platfrom, so set bContainsChunks true.
 	// This is a corner case but as far as I know there is no other situation that would allow for an empty bundle that was not chunked.
 	// If such a case were to arise, then bContainsChunks would need to be determined by each bundle source individually.
@@ -1334,6 +1364,26 @@ void FDefaultInstallBundleManager::UpdateBundleSources(FContentReleaseRequestRef
 		return;
 	}
 
+	FBundleInfo& BundleInfo = BundleInfoMap.FindChecked(Request->BundleName);
+
+	// Only unmount on demand tocs if this is a full release
+	if (BundleInfo.bMountedOnDemandTocs)
+	{
+		check(IOStoreOnDemandModule);
+		FIoStatus UnmountStatus = IOStoreOnDemandModule->Unmount(Request->BundleName.ToString());
+		if (!UnmountStatus.IsOk())
+		{
+			LOG_INSTALL_BUNDLE_MAN_OVERRIDE(Request->LogVerbosityOverride, Error, TEXT("Error unmounting ondemand toc for bundle '%s' : %s"), *BundleInfo.BundleNameString, *UnmountStatus.ToString());
+
+			// Use as a catch-all for something going wrong with IOStore, might want to add another code for this
+			Request->Result = EInstallBundleReleaseResult::ManifestArchiveError;
+			Request->StepResult = EContentRequestStepResult::Done;
+			return;
+		}
+
+		BundleInfo.bMountedOnDemandTocs = false;
+	}
+
 	// Release from any caches that were reserved
 	for (const TPair<FName, TSharedRef<FInstallBundleCache>>& Pair : BundleCaches)
 	{
@@ -1485,12 +1535,9 @@ void FDefaultInstallBundleManager::TickAsyncMountTasks()
 
 void FDefaultInstallBundleManager::MountPaks(FContentRequestRef Request)
 {
-	SCOPED_BOOT_TIMING("FDefaultInstallBundleManager::MountPaks");
-	TRACE_BOOKMARK(TEXT("Start Mount Bundle %s"), *Request->BundleName.ToString());
-
 	check(!Request->bIsCanceled && Request->Result == EInstallBundleResult::OK);
 
-	if (EnumHasAnyFlags(Request->Flags, EInstallBundleRequestFlags::SkipMount))
+	if (EnumHasAnyFlags(Request->Flags, EInstallBundleRequestFlags::SkipMount) && Request->OnDemandMountArgs.IsEmpty())
 	{
 		LOG_INSTALL_BUNDLE_MAN_OVERRIDE(Request->LogVerbosityOverride, Display, TEXT("Skipping Mounting Paks for Request %s"), *Request->BundleName.ToString());
 
@@ -1498,27 +1545,32 @@ void FDefaultInstallBundleManager::MountPaks(FContentRequestRef Request)
 		return;
 	}
 
-	FBundleInfo& BundleInfo = BundleInfoMap.FindChecked(Request->BundleName);
-
-	StatsBegin(Request->BundleName, EContentRequestState::Mounting);
+	SCOPED_BOOT_TIMING("FDefaultInstallBundleManager::MountPaks");
+	TRACE_BOOKMARK(TEXT("Start Mount Bundle %s"), *Request->BundleName.ToString());
 
-	LOG_INSTALL_BUNDLE_MAN_OVERRIDE(Request->LogVerbosityOverride, Display, TEXT("Mounting Paks for Request %s"), *BundleInfo.BundleNameString);
+	FBundleInfo& BundleInfo = BundleInfoMap.FindChecked(Request->BundleName);
 
-	TSharedRef<bool> bMountedPaks = MakeShared<bool>(false);
+	// Set waiting flag in case aysnc work happens
+	Request->StepResult = EContentRequestStepResult::Waiting;
 
-	TUniqueFunction<void()> WorkFunc =
+	auto MountPaksWorkFunc =
 		[ContentPaths = BundleInfo.ContentPaths.ContentPaths
-		, LogVerbosityOverride = Request->LogVerbosityOverride
-		, bMountedPaks]() mutable
+		, LogVerbosityOverride = Request->LogVerbosityOverride]() mutable
 	{
-		bMountedPaks.Get() = MountPaksInList(ContentPaths, LogVerbosityOverride);
+		return MountPaksInList(ContentPaths, LogVerbosityOverride);
 	};
 
-	TUniqueFunction<void()> OnCompleteFunc = [this, Request, bMountedPaks]()
+	auto OnPakMountCompleteFunc = [this, WeakThis = AsWeak(), Request](bool bMountedPaks) mutable
 	{
+		TSharedPtr<IInstallBundleManager> PinnedThis = WeakThis.Pin();
+		if (!PinnedThis)
+		{
+			return;
+		}
+
 		FBundleInfo& BundleInfoLambda = BundleInfoMap.FindChecked(Request->BundleName);
 
-		if (bMountedPaks.Get())
+		if (bMountedPaks)
 		{
 			FCoreDelegates::OnOptimizeMemoryUsageForMountedPaks.Execute();
 		}
@@ -1538,16 +1590,143 @@ void FDefaultInstallBundleManager::MountPaks(FContentRequestRef Request)
 		Request->StepResult = EContentRequestStepResult::Done;
 	};
 
-	bool bIsAsyncMount = EnumHasAnyFlags(Request->Flags, EInstallBundleRequestFlags::AsyncMount);
-	if (!bIsAsyncMount)
+	const bool bIsAsyncMount = EnumHasAnyFlags(Request->Flags, EInstallBundleRequestFlags::AsyncMount);
+	TUniqueFunction<void()> StartMountPaksFunc = 
+		[MountPaksWorkFunc = MoveTemp(MountPaksWorkFunc)
+		, OnPakMountCompleteFunc = MoveTemp(OnPakMountCompleteFunc)
+		, bIsAsyncMount]() mutable
 	{
-		WorkFunc();
-		OnCompleteFunc();
+		if (!bIsAsyncMount)
+		{
+			const bool bMountedPaks = MountPaksWorkFunc();
+			OnPakMountCompleteFunc(bMountedPaks);
+		}
+		else
+		{
+			InstallBundleUtil::StartInstallBundleAsyncIOTask(
+				[MountPaksWorkFunc = MoveTemp(MountPaksWorkFunc)
+				, OnPakMountCompleteFunc = MoveTemp(OnPakMountCompleteFunc)]() mutable
+			{
+				const bool bMountedPaks = MountPaksWorkFunc();
+				ExecuteOnGameThread(UE_SOURCE_LOCATION, [OnPakMountCompleteFunc = MoveTemp(OnPakMountCompleteFunc), bMountedPaks]() mutable
+				{
+					OnPakMountCompleteFunc(bMountedPaks);
+				});
+			});
+		}
+	};
+
+	/*
+	* 1 - Optionally, mount On Demand IOStore Tocs, this is always async
+	* 2 - Mount paks, respect EInstallBundleRequestFlags::AsyncMount
+	*/
+
+	if (Request->OnDemandMountArgs.IsEmpty() || BundleInfo.bMountedOnDemandTocs)
+	{
+		LOG_INSTALL_BUNDLE_MAN_OVERRIDE(Request->LogVerbosityOverride, Display, TEXT("Mounting Paks for Request %s"), *BundleInfo.BundleNameString);
+		StatsBegin(Request->BundleName, EContentRequestState::Mounting);
+		StartMountPaksFunc();
+		return;
 	}
-	else
+
+	if (!IOStoreOnDemandModule)
 	{
-		Request->StepResult = EContentRequestStepResult::Waiting;
-		InstallBundleUtil::StartInstallBundleAsyncIOTask(AsyncMountTasks, MoveTemp(WorkFunc), MoveTemp(OnCompleteFunc));
+		IOStoreOnDemandModule = FModuleManager::Get().GetModulePtr<UE::IoStore::FIoStoreOnDemandModule>(TEXT("IoStoreOnDemand"));
+		if (!IOStoreOnDemandModule)
+		{
+			Request->Result = EInstallBundleResult::InitializationError;
+			Request->OptionalErrorCode = TEXT("IoStoreOnDemand_Not_Found");
+			Request->StepResult = EContentRequestStepResult::Done;
+			return;
+		}
+	}
+
+	LOG_INSTALL_BUNDLE_MAN_OVERRIDE(Request->LogVerbosityOverride, Display, TEXT("Mounting OnDemand Tocs for Request %s"), *BundleInfo.BundleNameString);
+	StatsBegin(Request->BundleName, EContentRequestState::Mounting);
+
+	struct FOnDemandMountContext
+	{
+		FContentRequestPtr Request;
+		TArray<UE::IoStore::FOnDemandMountResult> Results;
+		TUniqueFunction<void()> StartMountPaksFunc;
+		TWeakPtr<FDefaultInstallBundleManager> WeakThis;
+
+		~FOnDemandMountContext()
+		{
+			ExecuteOnGameThread(UE_SOURCE_LOCATION, 
+			[Request = Request
+			, StartMountPaksFunc = MoveTemp(StartMountPaksFunc)
+			, Results = MoveTemp(Results)
+			, WeakThis = MoveTemp(WeakThis)]
+			{
+				TSharedPtr<FDefaultInstallBundleManager> PinnedThis = WeakThis.Pin();
+				if (!PinnedThis)
+				{
+					return;
+				}
+
+				for (const UE::IoStore::FOnDemandMountResult& MountResult : Results)
+				{
+					if (!MountResult.Status.IsOk() && MountResult.Status.GetErrorCode() != EIoErrorCode::PendingEncryptionKey)
+					{
+						FString StatusString = MountResult.Status.ToString();
+
+						LOG_INSTALL_BUNDLE_MAN(Error, TEXT("Error mounting on demand toc for '%s' : %s"), *Request->BundleName.ToString(), *StatusString);
+
+						FString ErrorCode = GetIoErrorText(MountResult.Status.GetErrorCode());
+						ErrorCode.ReplaceCharInline(TEXT(' '), TEXT('_'), ESearchCase::CaseSensitive);
+
+						PinnedThis->StatsEnd(Request->BundleName, EContentRequestState::Mounting);
+						TRACE_BOOKMARK(TEXT("Finsihed Mount Bundle %s"), *Request->BundleName.ToString());
+
+						Request->Result = EInstallBundleResult::InstallError;
+						Request->OptionalErrorCode = MoveTemp(ErrorCode);
+						Request->OptionalErrorText = FText::AsCultureInvariant(MountResult.Status.ToString()); // IOStore errors aren't localized 
+						Request->StepResult = EContentRequestStepResult::Done;
+						return;
+					}
+				}
+
+				FBundleInfo& BundleInfo = PinnedThis->BundleInfoMap.FindChecked(Request->BundleName);
+				BundleInfo.bMountedOnDemandTocs  = true;
+
+				if (EnumHasAnyFlags(Request->Flags, EInstallBundleRequestFlags::SkipMount))
+				{
+					LOG_INSTALL_BUNDLE_MAN_OVERRIDE(Request->LogVerbosityOverride, Display, TEXT("Skipping Mounting Paks for Request %s"), *Request->BundleName.ToString());
+
+					PinnedThis->StatsEnd(Request->BundleName, EContentRequestState::Mounting);
+					TRACE_BOOKMARK(TEXT("Finsihed Mount Bundle %s"), *Request->BundleName.ToString());
+					Request->StepResult = EContentRequestStepResult::Done;
+					return;
+				}
+
+				LOG_INSTALL_BUNDLE_MAN_OVERRIDE(Request->LogVerbosityOverride, Display, TEXT("Mounting Paks for Request %s"), *BundleInfo.BundleNameString);
+				StartMountPaksFunc();
+			});
+		}
+	};
+	TSharedRef<FOnDemandMountContext> OnDemandMountContext = MakeShared<FOnDemandMountContext>();
+	OnDemandMountContext->Request = Request;
+	OnDemandMountContext->StartMountPaksFunc = MoveTemp(StartMountPaksFunc);
+	OnDemandMountContext->WeakThis = StaticCastWeakPtr<FDefaultInstallBundleManager>(AsWeak());
+
+	for (TUniquePtr<UE::IoStore::FOnDemandMountArgs>& BundleMountArgs : Request->OnDemandMountArgs)
+	{
+		if (BundleMountArgs->MountId.IsEmpty())
+		{
+			BundleMountArgs->MountId = BundleInfo.BundleNameString;
+		}
+		else if (!ensureMsgf(BundleMountArgs->MountId == BundleInfo.BundleNameString, TEXT("Expected MountId to match bundle name!")))
+		{
+			BundleMountArgs->MountId = BundleInfo.BundleNameString;
+		}
+
+		BundleMountArgs->Options &= ~UE::IoStore::EOnDemandMountOptions::CallbackOnGameThread;
+
+		IOStoreOnDemandModule->Mount(MoveTemp(*BundleMountArgs), [OnDemandMountContext](UE::IoStore::FOnDemandMountResult OnDemandMountResult)
+		{
+			OnDemandMountContext->Results.Add(MoveTemp(OnDemandMountResult));
+		});
 	}
 }
 
@@ -1718,6 +1897,7 @@ void FDefaultInstallBundleManager::FinishRequest(FContentRequestRef Request)
 		ResultInfo.bIsStartup = BundleInfo.bIsStartup;
 		ResultInfo.bContentWasInstalled = Request->bContentWasInstalled;
 		ResultInfo.bContainsChunks = BundleInfo.ContentPaths.bContainsChunks;
+		ResultInfo.bContainsIoStoreOnDemand = BundleInfo.bMountedOnDemandTocs;
 		ResultInfo.OptionalErrorText = Request->OptionalErrorText;
 		ResultInfo.OptionalErrorCode = Request->OptionalErrorCode;
 
@@ -2805,6 +2985,13 @@ EInstallBundleSourceUpdateBundleInfoResult FDefaultInstallBundleManager::OnUpdat
 			}
 		}
 
+		// This wouldn't ever be expected to change for a particular source,
+		// but I don't have a good way of detected that.
+		if (SourceBundleInfo.bContainsIoStoreOnDemandToc)
+		{
+			BundleInfo.bContainsIoStoreOnDemandTocs = true;
+		}
+
 		if (SourceBundleInfo.Priority < BundleInfo.Priority)
 		{
 			BundleInfo.Priority = SourceBundleInfo.Priority;
@@ -4714,6 +4901,11 @@ void FDefaultInstallBundleManager::AsyncInit_OnQueryBundleInfoComplete(TSharedRe
 				Priority = SourceBundleInfo->Priority;
 			}
 
+			if (SourceBundleInfo->bContainsIoStoreOnDemandToc)
+			{
+				BundleInfo.bContainsIoStoreOnDemandTocs = true;
+			}
+
 			if (SourceBundleInfo->bIsStartup)
 			{
 				bIsStartup = true;
diff --git a/Engine/Plugins/Experimental/DefaultInstallBundleManager/Source/Private/InstallBundleManagerUtil.cpp b/Engine/Plugins/Experimental/DefaultInstallBundleManager/Source/Private/InstallBundleManagerUtil.cpp
--- a/Engine/Plugins/Experimental/DefaultInstallBundleManager/Source/Private/InstallBundleManagerUtil.cpp
+++ b/Engine/Plugins/Experimental/DefaultInstallBundleManager/Source/Private/InstallBundleManagerUtil.cpp
@@ -124,11 +124,6 @@ namespace InstallBundleManagerUtil
 			}
 		}
 
-		if (!InstallBundleConfig.GetBool(*Section, TEXT("UseChunkDBs"), OutInfo.bUseChunkDBs))
-		{
-			OutInfo.bUseChunkDBs = false;
-		}
-
 		return true;
 	}
 
diff --git a/Engine/Plugins/Experimental/DefaultInstallBundleManager/Source/Private/InstallBundleSourceBulk.cpp b/Engine/Plugins/Experimental/DefaultInstallBundleManager/Source/Private/InstallBundleSourceBulk.cpp
--- a/Engine/Plugins/Experimental/DefaultInstallBundleManager/Source/Private/InstallBundleSourceBulk.cpp
+++ b/Engine/Plugins/Experimental/DefaultInstallBundleManager/Source/Private/InstallBundleSourceBulk.cpp
@@ -2,6 +2,7 @@
 
 #include "InstallBundleSourceBulk.h"
 
+#include "Algo/AnyOf.h"
 #include "DefaultInstallBundleManagerPrivate.h"
 #include "HAL/PlatformFileManager.h"
 #include "HAL/PlatformFile.h"
@@ -501,6 +502,12 @@ void FInstallBundleSourceBulk::AsyncInit_QueryBundleInfo(FInstallBundleSourceQue
 
 			BundleInfo.BundleContentState = GetBundleInstallState(BundleInfo.BundleName);
 
+			if (TArray<FString>* FileList = BulkBuildBundles.Find(BundleInfo.BundleName))
+			{
+				BundleInfo.bContainsIoStoreOnDemandToc = Algo::AnyOf(
+					*FileList, [](FStringView File) { return File.EndsWith(TEXTVIEW(".uondemandtoc")); });
+			}
+
 			FName BundleName = BundleInfo.BundleName;
 			ResultInfo.SourceBundleInfoMap.Add(BundleName, MoveTemp(BundleInfo));
 		}
diff --git a/Engine/Plugins/Experimental/DefaultInstallBundleManager/Source/Private/InstallBundleSourcePlatformChunkInstall.cpp b/Engine/Plugins/Experimental/DefaultInstallBundleManager/Source/Private/InstallBundleSourcePlatformChunkInstall.cpp
--- a/Engine/Plugins/Experimental/DefaultInstallBundleManager/Source/Private/InstallBundleSourcePlatformChunkInstall.cpp
+++ b/Engine/Plugins/Experimental/DefaultInstallBundleManager/Source/Private/InstallBundleSourcePlatformChunkInstall.cpp
@@ -368,6 +368,9 @@ bool FInstallBundleSourcePlatformChunkInstall::QueryPersistentBundleInfo(FInstal
 		return false;
 	}
 
+	SourceBundleInfo.bContainsIoStoreOnDemandToc = Algo::AnyOf(
+		BundleInfo->FilePaths, [](FStringView File) { return File.EndsWith(TEXTVIEW(".uondemandtoc")); });
+
 	// assume the startup bundle is installed. It cannot be uninstalled
 	if (SourceBundleInfo.bIsStartup)
 	{
diff --git a/Engine/Plugins/Experimental/DefaultInstallBundleManager/Source/Public/DefaultInstallBundleManager.h b/Engine/Plugins/Experimental/DefaultInstallBundleManager/Source/Public/DefaultInstallBundleManager.h
--- a/Engine/Plugins/Experimental/DefaultInstallBundleManager/Source/Public/DefaultInstallBundleManager.h
+++ b/Engine/Plugins/Experimental/DefaultInstallBundleManager/Source/Public/DefaultInstallBundleManager.h
@@ -9,14 +9,14 @@
 #include "Containers/StaticArray.h"
 #include "Templates/ValueOrError.h"
 
-enum class EPatchCheckResult : uint8;
-
 #ifndef INSTALL_BUNDLE_ENABLE_ANALYTICS
 	#define INSTALL_BUNDLE_ENABLE_ANALYTICS (!WITH_EDITOR)
 #endif
 
-class IInstallBundleSource;
+enum class EPatchCheckResult : uint8;
 class IAnalyticsProviderET;
+class IInstallBundleSource;
+namespace UE::IoStore { class FIoStoreOnDemandModule; }
 
 class DEFAULTINSTALLBUNDLEMANAGER_API FDefaultInstallBundleManager : public IInstallBundleManager
 {
@@ -192,16 +192,20 @@ protected:
 		}
 	private:
 		EBundleState BundleStatus = EBundleState::NotInstalled;
-		bool bWaitForPSOCache = false;
 		uint32 InitialShaderPrecompiles = 0;
 	public:
 		FString BundleNameString; // Since FNames do not preserve casing
 		TArray<EBundlePrereqs> Prereqs;
+		TArray<FBundleSourceRelevance> ContributingSources; // Sources contributing to this bundle info		
+		FBundleContentPaths ContentPaths; // Only valid if BundleStatus >= NeedsMount
+		EInstallBundlePriority Priority = EInstallBundlePriority::Low;
+	private:
+		bool bWaitForPSOCache = false;
+	public:
 		bool bReleaseRequired = false; // A bundle source may have done install work so expects a release call
 		bool bIsStartup = false;
-		TArray<FBundleSourceRelevance> ContributingSources; // Sources contributing to this bundle info
-		EInstallBundlePriority Priority = EInstallBundlePriority::Low;
-		FBundleContentPaths ContentPaths; // Only valid if BundleStatus >= NeedsMount
+		bool bContainsIoStoreOnDemandTocs = false;
+		bool bMountedOnDemandTocs = false;
 	};
 
 	// GetBundleStatus protects erroneous accesses of the bundle status before initialization is complete by throwing an assert 
@@ -399,6 +403,9 @@ protected:
 		int RequiredSourceRequestResultsCount = 0;
 		// completion results from each bundle source
 		TMap<FInstallBundleSourceType, FInstallBundleSourceUpdateContentResultInfo> SourceRequestResults;
+
+		TArray<TUniquePtr<UE::IoStore::FOnDemandMountArgs>> OnDemandMountArgs;
+
 		FText OptionalErrorText;
 		FString OptionalErrorCode;
 
@@ -722,6 +729,7 @@ protected:
 	FTSTicker::FDelegateHandle TickHandle;
 	FDelegateHandle AsyncInit_PatchCheckHandle;
 	FDelegateHandle PatchCheckHandle;
+	UE::IoStore::FIoStoreOnDemandModule* IOStoreOnDemandModule = nullptr;
 	FInstallBundleSourceFactoryFunction InstallBundleSourceFactory;
 
 	TMap<FName, FBundleInfo> BundleInfoMap;
diff --git a/Engine/Source/Runtime/InstallBundleManager/InstallBundleManager.Build.cs b/Engine/Source/Runtime/InstallBundleManager/InstallBundleManager.Build.cs
--- a/Engine/Source/Runtime/InstallBundleManager/InstallBundleManager.Build.cs
+++ b/Engine/Source/Runtime/InstallBundleManager/InstallBundleManager.Build.cs
@@ -8,8 +8,9 @@ public class InstallBundleManager : ModuleRules
 	{
 		PrivateDependencyModuleNames.AddRange(
 			new string[] {
-				"Core",
 				"ApplicationCore",
+				"Core",
+				"IoStoreOnDemand",
 				"Json"
 			}
 		);
diff --git a/Engine/Source/Runtime/InstallBundleManager/Private/InstallBundleTypes.cpp b/Engine/Source/Runtime/InstallBundleManager/Private/InstallBundleTypes.cpp
--- a/Engine/Source/Runtime/InstallBundleManager/Private/InstallBundleTypes.cpp
+++ b/Engine/Source/Runtime/InstallBundleManager/Private/InstallBundleTypes.cpp
@@ -6,6 +6,7 @@
 #include "InstallBundleManagerPrivate.h"
 #include "InstallBundleUtils.h"
 #include "InstallBundleManagerInterface.h"
+#include "IO/IoStoreOnDemand.h"
 #include "Misc/CString.h"
 
 PRAGMA_DISABLE_DEPRECATION_WARNINGS
@@ -311,3 +312,7 @@ bool FInstallBundleCombinedContentState::GetAnyBundleHasState(EInstallBundleInst
 
 	return false;
 }
+
+FInstallBundleSourceUpdateContentResultInfo::FInstallBundleSourceUpdateContentResultInfo() = default;
+FInstallBundleSourceUpdateContentResultInfo::FInstallBundleSourceUpdateContentResultInfo(FInstallBundleSourceUpdateContentResultInfo&&) = default;
+FInstallBundleSourceUpdateContentResultInfo::~FInstallBundleSourceUpdateContentResultInfo() = default;
diff --git a/Engine/Source/Runtime/InstallBundleManager/Private/InstallBundleUtils.cpp b/Engine/Source/Runtime/InstallBundleManager/Private/InstallBundleUtils.cpp
--- a/Engine/Source/Runtime/InstallBundleManager/Private/InstallBundleUtils.cpp
+++ b/Engine/Source/Runtime/InstallBundleManager/Private/InstallBundleUtils.cpp
@@ -34,6 +34,42 @@ namespace InstallBundleUtil
 			&& ConnectionType != ENetworkConnectionType::None;
 	}
 
+	bool SplitHostUrl(const FStringView Url, FStringView& OutHost, FStringView& OutRemainder)
+	{
+		OutHost = OutRemainder = FStringView();
+
+		const FStringView ProtoHttp = TEXTVIEW("http://");
+		const FStringView ProtoHttps = TEXTVIEW("https://");
+
+		FStringView Proto;
+		if (Url.StartsWith(ProtoHttp))
+		{
+			Proto = ProtoHttp;
+		}
+		else if (Url.StartsWith(ProtoHttps))
+		{
+			Proto = ProtoHttps;
+		}
+		else
+		{
+			return false;
+		}
+
+		FStringView ProtoChopped = Url.RightChop(Proto.Len());
+
+		int32 DelimPos = INDEX_NONE;
+		if (!ProtoChopped.FindChar(TEXT('/'), DelimPos))
+		{
+			OutHost = ProtoChopped;
+			return !OutHost.IsEmpty();
+		}
+
+		OutHost = Url.Left(Proto.Len() + DelimPos);
+		OutRemainder = Url.RightChop(OutHost.Len());
+
+		return true;
+	}
+
 	const TCHAR* GetInstallBundlePauseReason(EInstallBundlePauseFlags Flags)
 	{
 		// Return the most appropriate reason given the flags
diff --git a/Engine/Source/Runtime/InstallBundleManager/Public/InstallBundleManagerInterface.h b/Engine/Source/Runtime/InstallBundleManager/Public/InstallBundleManagerInterface.h
--- a/Engine/Source/Runtime/InstallBundleManager/Public/InstallBundleManagerInterface.h
+++ b/Engine/Source/Runtime/InstallBundleManager/Public/InstallBundleManagerInterface.h
@@ -50,6 +50,7 @@ struct FInstallBundleRequestResultInfo
 	EInstallBundleResult Result = EInstallBundleResult::OK;
 	bool bIsStartup = false;
 	bool bContainsChunks = false;
+	bool bContainsIoStoreOnDemand = false;
 	bool bContentWasInstalled = false;
 
 	// Currently, these just forward BPT Error info
diff --git a/Engine/Source/Runtime/InstallBundleManager/Public/InstallBundleTypes.h b/Engine/Source/Runtime/InstallBundleManager/Public/InstallBundleTypes.h
--- a/Engine/Source/Runtime/InstallBundleManager/Public/InstallBundleTypes.h
+++ b/Engine/Source/Runtime/InstallBundleManager/Public/InstallBundleTypes.h
@@ -112,6 +112,7 @@ INSTALLBUNDLEMANAGER_API const TCHAR* LexToString(EInstallBundleInstallState Sta
 struct FInstallBundleCombinedInstallState
 {
 	TMap<FName, EInstallBundleInstallState> IndividualBundleStates;
+	TSet<FName> BundlesWithIoStoreOnDemand;
 
 	INSTALLBUNDLEMANAGER_API bool GetAllBundlesHaveState(EInstallBundleInstallState State, TArrayView<const FName> ExcludedBundles = TArrayView<const FName>()) const;
 	INSTALLBUNDLEMANAGER_API bool GetAnyBundleHasState(EInstallBundleInstallState State, TArrayView<const FName> ExcludedBundles = TArrayView<const FName>()) const;
@@ -148,6 +149,8 @@ struct FInstallBundleContentSize
 	}
 };
 
+// TODO: Create a per-source version of this struct so that its clear what data bundle sources
+// are responsible for populating.
 struct FInstallBundleCombinedContentState
 {
 	PRAGMA_DISABLE_DEPRECATION_WARNINGS
@@ -160,6 +163,8 @@ struct FInstallBundleCombinedContentState
 	PRAGMA_ENABLE_DEPRECATION_WARNINGS
 	
 	TMap<FName, FInstallBundleContentState> IndividualBundleStates;
+	// Bundle Sources don't need to populate this, bundle manager can determine it from init data
+	TSet<FName> BundlesWithIoStoreOnDemand;
 	TMap<FInstallBundleSourceType, FString> CurrentVersion;
 	
 #if !UE_BUILD_SHIPPING
@@ -328,8 +333,7 @@ struct FInstallBundleSourceUpdateBundleInfo
 	FDateTime LastAccessTime = FDateTime::MinValue(); // If cached, used to decide eviction order
 	EInstallBundleInstallState BundleContentState = EInstallBundleInstallState::NotInstalled; // Whether this bundle is up to date
 	bool bIsCached = false; // Whether this bundle should be cached if this source has a bundle cache
-	// TODO: this appears to be unused
-	bool bUseChunkDBs = false; // Whether this bundle should attempt to use ChunkDBs at all
+	bool bContainsIoStoreOnDemandToc = false; // Whether this bundle contains an an IoStoreOnDemand Toc
 };
 
 struct FInstallBundleSourceUpdateBundleInfoResult
@@ -361,8 +365,18 @@ enum class EInstallBundleSourceUpdateBundleInfoResult : uint8
 };
 INSTALLBUNDLEMANAGER_API const TCHAR* LexToString(EInstallBundleSourceUpdateBundleInfoResult Result);
 
+namespace UE::IoStore
+{
+	struct FOnDemandMountArgs;
+}
+
 struct FInstallBundleSourceUpdateContentResultInfo
 {
+	INSTALLBUNDLEMANAGER_API FInstallBundleSourceUpdateContentResultInfo();
+	FInstallBundleSourceUpdateContentResultInfo(const FInstallBundleSourceUpdateContentResultInfo&) = delete;
+	INSTALLBUNDLEMANAGER_API FInstallBundleSourceUpdateContentResultInfo(FInstallBundleSourceUpdateContentResultInfo&&);
+	INSTALLBUNDLEMANAGER_API ~FInstallBundleSourceUpdateContentResultInfo();
+
 	FName BundleName;
 	EInstallBundleResult Result = EInstallBundleResult::OK;
 
@@ -375,6 +389,7 @@ struct FInstallBundleSourceUpdateContentResultInfo
 	TArray<FString> AdditionalRootDirs;
 	// Support platforms that need shaderlibs in the physical FS
 	TSet<FString> NonUFSShaderLibPaths;
+	TArray<TUniquePtr<UE::IoStore::FOnDemandMountArgs>> OnDemandMountArgs;
 	FString ProjectName;
 
 	uint64 CurrentInstallSize = 0;
diff --git a/Engine/Source/Runtime/InstallBundleManager/Public/InstallBundleUtils.h b/Engine/Source/Runtime/InstallBundleManager/Public/InstallBundleUtils.h
--- a/Engine/Source/Runtime/InstallBundleManager/Public/InstallBundleUtils.h
+++ b/Engine/Source/Runtime/InstallBundleManager/Public/InstallBundleUtils.h
@@ -33,6 +33,8 @@ namespace InstallBundleUtil
 
 	INSTALLBUNDLEMANAGER_API bool HasInternetConnection(ENetworkConnectionType ConnectionType);
 
+	INSTALLBUNDLEMANAGER_API bool SplitHostUrl(const FStringView Url, FStringView& OutHost, FStringView& OutRemainder);
+
 	INSTALLBUNDLEMANAGER_API const TCHAR* GetInstallBundlePauseReason(EInstallBundlePauseFlags Flags);
 
 	INSTALLBUNDLEMANAGER_API const FString& GetInstallBundleSectionPrefix();
