Android nonunity CIS fix

diff --git a/Engine/Plugins/Experimental/DefaultInstallBundleManager/Source/Private/DefaultInstallBundleManager.cpp b/Engine/Plugins/Experimental/DefaultInstallBundleManager/Source/Private/DefaultInstallBundleManager.cpp
--- a/Engine/Plugins/Experimental/DefaultInstallBundleManager/Source/Private/DefaultInstallBundleManager.cpp
+++ b/Engine/Plugins/Experimental/DefaultInstallBundleManager/Source/Private/DefaultInstallBundleManager.cpp
@@ -186,6 +186,8 @@ void LoadKeychainFromInI(FKeyChain& OutCryptoSettings)
 }
 #endif // WITH_EDITOR
 
+FDefaultInstallBundleManager::FContentRequest::~FContentRequest() = default;
+
 FDefaultInstallBundleManager::FDefaultInstallBundleManager(FInstallBundleSourceFactoryFunction InInstallBundleSourceFactory)
 	: InstallBundleSourceFactory(InInstallBundleSourceFactory ? MoveTemp(InInstallBundleSourceFactory) : InstallBundleManagerUtil::MakeBundleSource )
 	, PersistentStats(MakeShared<InstallBundleManagerUtil::FPersistentStatContainer>())
diff --git a/Engine/Plugins/Experimental/DefaultInstallBundleManager/Source/Public/DefaultInstallBundleManager.h b/Engine/Plugins/Experimental/DefaultInstallBundleManager/Source/Public/DefaultInstallBundleManager.h
--- a/Engine/Plugins/Experimental/DefaultInstallBundleManager/Source/Public/DefaultInstallBundleManager.h
+++ b/Engine/Plugins/Experimental/DefaultInstallBundleManager/Source/Public/DefaultInstallBundleManager.h
@@ -371,6 +371,8 @@ protected:
 
 	struct FContentRequest : FCacheEvictionRequestor
 	{
+		~FContentRequest();
+
 		EContentRequestStepResult StepResult = EContentRequestStepResult::Done;
 		TArray<EContentRequestState> Steps;
 		int32 iStep = INDEX_NONE;
