Fix InstallBundleSource setting bContentWasInstalled when it didn't actually do work. This is throwing off our cache analytics

diff --git a/Engine/Plugins/Experimental/DefaultInstallBundleManager/Source/Private/InstallBundleSourcePlatformChunkInstall.cpp b/Engine/Plugins/Experimental/DefaultInstallBundleManager/Source/Private/InstallBundleSourcePlatformChunkInstall.cpp
--- a/Engine/Plugins/Experimental/DefaultInstallBundleManager/Source/Private/InstallBundleSourcePlatformChunkInstall.cpp
+++ b/Engine/Plugins/Experimental/DefaultInstallBundleManager/Source/Private/InstallBundleSourcePlatformChunkInstall.cpp
@@ -510,7 +510,7 @@ void FInstallBundleSourcePlatformChunkInstall::RequestUpdateContent(FRequestUpda
 		ResultInfo.BundleName = Context.BundleName;
 		ResultInfo.Result = EInstallBundleResult::OK;
 		ResultInfo.ContentPaths = BundleInfo->FilePaths;
-		ResultInfo.bContentWasInstalled = (BundleInfo->FilePaths.Num() > 0);
+		ResultInfo.bContentWasInstalled = false;
 		Context.CompleteCallback.ExecuteIfBound(AsShared(), MoveTemp(ResultInfo));
 	}
 	else
