CIS fix.

#rnx

diff --git a/Engine/Plugins/Experimental/DefaultInstallBundleManager/Source/Private/InstallBundleSourcePlatformChunkInstall.cpp b/Engine/Plugins/Experimental/DefaultInstallBundleManager/Source/Private/InstallBundleSourcePlatformChunkInstall.cpp
--- a/Engine/Plugins/Experimental/DefaultInstallBundleManager/Source/Private/InstallBundleSourcePlatformChunkInstall.cpp
+++ b/Engine/Plugins/Experimental/DefaultInstallBundleManager/Source/Private/InstallBundleSourcePlatformChunkInstall.cpp
@@ -270,7 +270,7 @@ void FInstallBundleSourcePlatformChunkInstall::AsyncInit(FInstallBundleSourceIni
 		{
 			if (bHasAbsoluteRootDir && FPaths::IsRelative(PakFolder))
 			{
-				PakSearchDirs.Add_GetRef(MoveTemp(PakFolder));
+				PakSearchDirs.Add(MoveTemp(PakFolder));
 			}
 			else if (PakFolder.StartsWith(RootDir))
 			{
