Add support to bulk bundle source for mounting uondemandtocs.

diff --git a/Engine/Plugins/Experimental/DefaultInstallBundleManager/Source/Private/InstallBundleSourceBulk.cpp b/Engine/Plugins/Experimental/DefaultInstallBundleManager/Source/Private/InstallBundleSourceBulk.cpp
--- a/Engine/Plugins/Experimental/DefaultInstallBundleManager/Source/Private/InstallBundleSourceBulk.cpp
+++ b/Engine/Plugins/Experimental/DefaultInstallBundleManager/Source/Private/InstallBundleSourceBulk.cpp
@@ -6,6 +6,7 @@
 #include "DefaultInstallBundleManagerPrivate.h"
 #include "HAL/PlatformFileManager.h"
 #include "HAL/PlatformFile.h"
+#include "IO/IoStoreOnDemand.h"
 #include "IPlatformFilePak.h"
 #include "InstallBundleManagerUtil.h"
 #include "Misc/CommandLine.h"
@@ -443,6 +444,17 @@ void FInstallBundleSourceBulk::SerializeBulkBuildBundleMetadata(const TMap<FName
 	}
 }
 
+void FInstallBundleSourceBulk::GetOnDemandHostGroup(UE::IoStore::FOnDemandHostGroup& OutHostGroup)
+{
+	ensureAlwaysMsgf(false, TEXT("FInstallBundleSourceBulk::GetOnDemandHostGroup not implemented!"));
+}
+
+FString FInstallBundleSourceBulk::GetOnDemandTocRelativeURL()
+{
+	ensureAlwaysMsgf(false, TEXT("FInstallBundleSourceBulk::GetOnDemandTocRelativeURL not implemented!"));
+	return {};
+}
+
 EInstallBundleInstallState FInstallBundleSourceBulk::GetBundleInstallState(FName BundleName)
 {
 	return EInstallBundleInstallState::UpToDate;
@@ -564,6 +576,7 @@ void FInstallBundleSourceBulk::RequestUpdateContent(FRequestUpdateContentBundleC
 
 	TArray<FString> NonUFSShaderLibPathsExtensions;
 	GConfig->GetArray(TEXT("InstallBundleManager.Shaderlib"), TEXT("ExtensionList"), NonUFSShaderLibPathsExtensions, GInstallBundleIni);
+
 	for (const FString& Path : ResultInfo.ContentPaths)
 	{
 		for (const FString& Ext : NonUFSShaderLibPathsExtensions)
@@ -574,6 +587,18 @@ void FInstallBundleSourceBulk::RequestUpdateContent(FRequestUpdateContentBundleC
 				ResultInfo.NonUFSShaderLibPaths.Add(FPaths::GetPath(Path));
 			}
 		}
+
+		if (Path.EndsWith(TEXTVIEW(".uondemandtoc")))
+		{
+			TUniquePtr<UE::IoStore::FOnDemandMountArgs> MountArgs = MakeUnique<UE::IoStore::FOnDemandMountArgs>();
+			MountArgs->MountId = Context.BundleName.ToString();
+			GetOnDemandHostGroup(MountArgs->HostGroup);
+			MountArgs->TocRelativeUrl = GetOnDemandTocRelativeURL();
+			MountArgs->FilePath = Path;
+			MountArgs->Options = UE::IoStore::EOnDemandMountOptions::InstallOnDemand;
+
+			ResultInfo.OnDemandMountArgs.Add(MoveTemp(MountArgs));
+		}
 	}
 	
 	ResultInfo.ProjectName = FApp::GetProjectName();
diff --git a/Engine/Plugins/Experimental/DefaultInstallBundleManager/Source/Public/InstallBundleSourceBulk.h b/Engine/Plugins/Experimental/DefaultInstallBundleManager/Source/Public/InstallBundleSourceBulk.h
--- a/Engine/Plugins/Experimental/DefaultInstallBundleManager/Source/Public/InstallBundleSourceBulk.h
+++ b/Engine/Plugins/Experimental/DefaultInstallBundleManager/Source/Public/InstallBundleSourceBulk.h
@@ -5,6 +5,8 @@
 #include "InstallBundleSourceInterface.h"
 #include "InstallBundleUtils.h"
 
+namespace UE::IoStore{ class FOnDemandHostGroup; }
+
 class DEFAULTINSTALLBUNDLEMANAGER_API FInstallBundleSourceBulk : public IInstallBundleSource
 {
 public:
@@ -58,6 +60,11 @@ public:
 	//Serialize out our BulkBundleBuild information to a BulkBundleBuild.ini file for future runs to not have to parse this information
 	static void SerializeBulkBuildBundleMetadata(const TMap<FName, TArray<FString>>& BulkBuildBundles);
 
+	// Support for on demand tocs, must be implemented by the game
+protected:
+	virtual void GetOnDemandHostGroup(UE::IoStore::FOnDemandHostGroup& OutHostGroup);
+	virtual FString GetOnDemandTocRelativeURL();
+
 protected:
 	FTSTicker::FDelegateHandle TickHandle;
 
