Allow ondemandtoc loading to be disabled by game

diff --git a/Engine/Plugins/Experimental/DefaultInstallBundleManager/Source/Private/DefaultInstallBundleManager.cpp b/Engine/Plugins/Experimental/DefaultInstallBundleManager/Source/Private/DefaultInstallBundleManager.cpp
--- a/Engine/Plugins/Experimental/DefaultInstallBundleManager/Source/Private/DefaultInstallBundleManager.cpp
+++ b/Engine/Plugins/Experimental/DefaultInstallBundleManager/Source/Private/DefaultInstallBundleManager.cpp
@@ -1623,7 +1623,7 @@ void FDefaultInstallBundleManager::MountPaks(FContentRequestRef Request)
 	* 2 - Mount paks, respect EInstallBundleRequestFlags::AsyncMount
 	*/
 
-	if (Request->OnDemandMountArgs.IsEmpty() || BundleInfo.bMountedOnDemandTocs)
+	if (Request->OnDemandMountArgs.IsEmpty() || BundleInfo.bMountedOnDemandTocs || !AllowIoStoreOnDemandMount(Request, BundleInfo))
 	{
 		LOG_INSTALL_BUNDLE_MAN_OVERRIDE(Request->LogVerbosityOverride, Display, TEXT("Mounting Paks for Request %s"), *BundleInfo.BundleNameString);
 		StatsBegin(Request->BundleName, EContentRequestState::Mounting);
@@ -1690,7 +1690,7 @@ void FDefaultInstallBundleManager::MountPaks(FContentRequestRef Request)
 				}
 
 				FBundleInfo& BundleInfo = PinnedThis->BundleInfoMap.FindChecked(Request->BundleName);
-				BundleInfo.bMountedOnDemandTocs  = true;
+				BundleInfo.bMountedOnDemandTocs = true;
 
 				if (EnumHasAnyFlags(Request->Flags, EInstallBundleRequestFlags::SkipMount))
 				{
@@ -1824,6 +1824,11 @@ void FDefaultInstallBundleManager::UnmountPaks(FContentReleaseRequestRef Request
 	Request->StepResult = EContentRequestStepResult::Done;
 }
 
+bool FDefaultInstallBundleManager::AllowIoStoreOnDemandMount(FContentRequestRef Request, FBundleInfo& BundleInfo)
+{
+	return true;
+}
+
 void FDefaultInstallBundleManager::WaitForShaderCache(FContentRequestRef Request)
 {
 	if (EnumHasAnyFlags(Request->Flags, EInstallBundleRequestFlags::SkipMount))
diff --git a/Engine/Plugins/Experimental/DefaultInstallBundleManager/Source/Public/DefaultInstallBundleManager.h b/Engine/Plugins/Experimental/DefaultInstallBundleManager/Source/Public/DefaultInstallBundleManager.h
--- a/Engine/Plugins/Experimental/DefaultInstallBundleManager/Source/Public/DefaultInstallBundleManager.h
+++ b/Engine/Plugins/Experimental/DefaultInstallBundleManager/Source/Public/DefaultInstallBundleManager.h
@@ -564,6 +564,8 @@ protected:
 
 	void UnmountPaks(FContentReleaseRequestRef Request);
 
+	virtual bool AllowIoStoreOnDemandMount(FContentRequestRef Request, FBundleInfo& BundleInfo);
+
 	virtual void OnPaksMountedInternal(FContentRequestRef Request, FBundleInfo& BundleInfo) {}
 	virtual void OnPaksUnmountedInternal(FContentReleaseRequestRef Request, FBundleInfo& BundleInfo) {};
 
