DefaultInstallBundleManager
- Fixed ondemandtoc possibly not being mounted during GFP install state (SkipMount is set)

review-39153233 [REDACTED_HANDLE]

diff --git a/Engine/Plugins/Experimental/DefaultInstallBundleManager/Source/Private/DefaultInstallBundleManager.cpp b/Engine/Plugins/Experimental/DefaultInstallBundleManager/Source/Private/DefaultInstallBundleManager.cpp
--- a/Engine/Plugins/Experimental/DefaultInstallBundleManager/Source/Private/DefaultInstallBundleManager.cpp
+++ b/Engine/Plugins/Experimental/DefaultInstallBundleManager/Source/Private/DefaultInstallBundleManager.cpp
@@ -3330,7 +3330,6 @@ TValueOrError<FInstallBundleRequestInfo, EInstallBundleResult> FDefaultInstallBu
 			// canceled release has finished.
 			if (!bCanceledRelease && EnumHasAnyFlags(Flags, EInstallBundleRequestFlags::SkipMount) && GetBundleStatus(*BundleInfo) == EBundleState::NeedsMount)
 			{
-				// If this bundle is not reserved in a cache, an install request cannot be skipped
 				bool bNeedsCacheReserve = false;
 				for (const FBundleSourceRelevance& SourceRelevance : BundleInfo->ContributingSources)
 				{
@@ -3345,7 +3344,9 @@ TValueOrError<FInstallBundleRequestInfo, EInstallBundleResult> FDefaultInstallBu
 					}
 				}
 
-				if (!bNeedsCacheReserve)
+				// If this bundle is not reserved in a cache but uses a cache, an install request cannot be skipped
+				// If this bundle uses IoStoreOnDemandTocs, an install request cannot be skipped
+				if (!bNeedsCacheReserve && !BundleInfo->bContainsIoStoreOnDemandTocs)
 				{
 					RetInfo.InfoFlags |= EInstallBundleRequestInfoFlags::SkippedAlreadyUpdatedBundles;
 					LOG_INSTALL_BUNDLE_MAN_OVERRIDE(LogVerbosityOverride, Verbose, TEXT("RequestUpdateContent Bundle %s  - Already Updated"), *BundleInfo->BundleNameString);
