InstallBundleSourceBulk
- Fix issue where uondemandtoc files aren't found if bOnlyGatherPaksInBulkData is set.

review-39153233 [REDACTED_HANDLE]

diff --git a/Engine/Plugins/Experimental/DefaultInstallBundleManager/Source/Private/InstallBundleSourceBulk.cpp b/Engine/Plugins/Experimental/DefaultInstallBundleManager/Source/Private/InstallBundleSourceBulk.cpp
--- a/Engine/Plugins/Experimental/DefaultInstallBundleManager/Source/Private/InstallBundleSourceBulk.cpp
+++ b/Engine/Plugins/Experimental/DefaultInstallBundleManager/Source/Private/InstallBundleSourceBulk.cpp
@@ -268,12 +268,55 @@ void FInstallBundleSourceBulk::AsyncInit_MakeBundlesForBulkBuild()
 	InstallBundleUtil::StartInstallBundleAsyncIOTask(InitAsyncTasks,
 	[FoundFiles, PakSearchDirs=MoveTemp(PakSearchDirs), ContentDir=FPaths::ProjectContentDir(), bOnlyGatherPaksInBulkData]()
 	{
-		const FString PakFileExtension(TEXT(".pak"));
-		const TCHAR* FileExtension = bOnlyGatherPaksInBulkData ? *PakFileExtension : nullptr;
-		
+		// Custom visitor because we need to deal with multiple extensions
+		class FFindBulkFilesVisitor : public IPlatformFile::FDirectoryVisitor
+		{
+		public:
+			FRWLock FoundFilesLock;
+			TArray<FString>& FoundFiles;
+			bool bOnlyGatherPaksInBulkData;
+			FFindBulkFilesVisitor(TArray<FString>& InFoundFiles, bool bInOnlyGatherPaksInBulkData)
+				: IPlatformFile::FDirectoryVisitor(EDirectoryVisitorFlags::ThreadSafe)
+				, FoundFiles(InFoundFiles)
+				, bOnlyGatherPaksInBulkData(bInOnlyGatherPaksInBulkData)
+			{}
+
+			virtual bool Visit(const TCHAR* FilenameOrDirectory, bool bIsDirectory) override
+			{
+				const FStringView FileExtensions[] = { TEXTVIEW(".pak"), TEXTVIEW(".uondemandtoc") };
+
+				if (!bIsDirectory)
+				{
+					bool bFound = !bOnlyGatherPaksInBulkData;
+					if (!bFound)
+					{
+						FStringView PathExtension = FPathViews::GetExtension(FilenameOrDirectory, true);
+						for (FStringView Extension : FileExtensions)
+						{
+							if (PathExtension == Extension)
+							{
+								bFound = true;
+								break;
+							}
+						}
+					}
+
+					if (bFound)
+					{
+						FString FileName(FilenameOrDirectory);
+						FRWScopeLock ScopeLock(FoundFilesLock, SLT_Write);
+						FoundFiles.Emplace(MoveTemp(FileName));
+					}
+				}
+
+				return true;
+			}
+		};
+
+		FFindBulkFilesVisitor Visitor(*FoundFiles, bOnlyGatherPaksInBulkData);
 		for (const FString& SearchDir : PakSearchDirs)
 		{
-			FPlatformFileManager::Get().GetPlatformPhysical().FindFilesRecursively(*FoundFiles, *SearchDir, FileExtension);
+			FPlatformFileManager::Get().GetPlatformPhysical().IterateDirectoryRecursively(*SearchDir, Visitor);
 		}
 	},
 	[this, FoundFiles]()
