DefaultInstallBundleManager
- Allow overriding pak mount list

review-39166808 [REDACTED_HANDLE]

diff --git a/Engine/Plugins/Experimental/DefaultInstallBundleManager/Source/Private/DefaultInstallBundleManager.cpp b/Engine/Plugins/Experimental/DefaultInstallBundleManager/Source/Private/DefaultInstallBundleManager.cpp
--- a/Engine/Plugins/Experimental/DefaultInstallBundleManager/Source/Private/DefaultInstallBundleManager.cpp
+++ b/Engine/Plugins/Experimental/DefaultInstallBundleManager/Source/Private/DefaultInstallBundleManager.cpp
@@ -1556,7 +1556,7 @@ void FDefaultInstallBundleManager::MountPaks(FContentRequestRef Request)
 	Request->StepResult = EContentRequestStepResult::Waiting;
 
 	auto MountPaksWorkFunc =
-		[ContentPaths = BundleInfo.ContentPaths.ContentPaths
+		[ContentPaths = GetPakMountList(Request, BundleInfo)
 		, LogVerbosityOverride = Request->LogVerbosityOverride]() mutable
 	{
 		return MountPaksInList(ContentPaths, LogVerbosityOverride);
@@ -1824,11 +1824,25 @@ void FDefaultInstallBundleManager::UnmountPaks(FContentReleaseRequestRef Request
 	Request->StepResult = EContentRequestStepResult::Done;
 }
 
-bool FDefaultInstallBundleManager::AllowIoStoreOnDemandMount(FContentRequestRef Request, FBundleInfo& BundleInfo)
+bool FDefaultInstallBundleManager::AllowIoStoreOnDemandMount(FContentRequestRef Request, const FBundleInfo& BundleInfo)
 {
 	return true;
 }
 
+TArray<FString> FDefaultInstallBundleManager::GetPakMountList(FContentRequestRef Request, const FBundleInfo& BundleInfo)
+{
+	TArray<FString> PakList;
+	for (const FString& Path : BundleInfo.ContentPaths.ContentPaths)
+	{
+		if (Path.EndsWith(TEXTVIEW(".pak")))
+		{
+			PakList.Add(Path);
+		}
+	}
+
+	return PakList;
+}
+
 void FDefaultInstallBundleManager::WaitForShaderCache(FContentRequestRef Request)
 {
 	if (EnumHasAnyFlags(Request->Flags, EInstallBundleRequestFlags::SkipMount))
diff --git a/Engine/Plugins/Experimental/DefaultInstallBundleManager/Source/Public/DefaultInstallBundleManager.h b/Engine/Plugins/Experimental/DefaultInstallBundleManager/Source/Public/DefaultInstallBundleManager.h
--- a/Engine/Plugins/Experimental/DefaultInstallBundleManager/Source/Public/DefaultInstallBundleManager.h
+++ b/Engine/Plugins/Experimental/DefaultInstallBundleManager/Source/Public/DefaultInstallBundleManager.h
@@ -564,7 +564,9 @@ protected:
 
 	void UnmountPaks(FContentReleaseRequestRef Request);
 
-	virtual bool AllowIoStoreOnDemandMount(FContentRequestRef Request, FBundleInfo& BundleInfo);
+	virtual bool AllowIoStoreOnDemandMount(FContentRequestRef Request, const FBundleInfo& BundleInfo);
+
+	virtual TArray<FString> GetPakMountList(FContentRequestRef Request, const FBundleInfo& BundleInfo);
 
 	virtual void OnPaksMountedInternal(FContentRequestRef Request, FBundleInfo& BundleInfo) {}
 	virtual void OnPaksUnmountedInternal(FContentReleaseRequestRef Request, FBundleInfo& BundleInfo) {};
