InstallBundleManager
- Remove NonUFSShaderLibPaths since its no longer used by any platform

diff --git a/Engine/Config/BaseInstallBundle.ini b/Engine/Config/BaseInstallBundle.ini
--- a/Engine/Config/BaseInstallBundle.ini
+++ b/Engine/Config/BaseInstallBundle.ini
@@ -8,7 +8,3 @@
 bShouldSerializeMissingBulkBuildDataIni=false
 ; If we limit our file parsing in Bulk bundle sources to just .pak files
 bOnlyGatherPaksInBulkData=false
-
-[InstallBundleManager.Shaderlib]
-; List of extensions that will get added to NonUFSShaderLibPaths
-!ExtensionList=ClearArray
\ No newline at end of file
diff --git a/Engine/Config/IOS/BaseIOSInstallBundle.ini b/Engine/Config/IOS/BaseIOSInstallBundle.ini
--- a/Engine/Config/IOS/BaseIOSInstallBundle.ini
+++ b/Engine/Config/IOS/BaseIOSInstallBundle.ini
@@ -5,7 +5,3 @@
 bShouldSerializeMissingBulkBuildDataIni=true
 ; If we limit our file parsing in Bulk bundle sources to just .pak files
 bOnlyGatherPaksInBulkData=true
-
-[InstallBundleManager.Shaderlib]
-; List of extensions that will get added to NonUFSShaderLibPaths
-!ExtensionList=ClearArray
\ No newline at end of file
diff --git a/Engine/Plugins/Experimental/DefaultInstallBundleManager/Source/Private/DefaultInstallBundleManager.cpp b/Engine/Plugins/Experimental/DefaultInstallBundleManager/Source/Private/DefaultInstallBundleManager.cpp
--- a/Engine/Plugins/Experimental/DefaultInstallBundleManager/Source/Private/DefaultInstallBundleManager.cpp
+++ b/Engine/Plugins/Experimental/DefaultInstallBundleManager/Source/Private/DefaultInstallBundleManager.cpp
@@ -1296,7 +1296,6 @@ void FDefaultInstallBundleManager::UpdateBundleSourceComplete(TSharedRef<IInstal
 			Request->OnDemandMountArgs.Append(MoveTemp(ResultInfo.OnDemandMountArgs));
 
 			BundleInfo.ContentPaths.AdditionalRootDirs.Append(MoveTemp(ResultInfo.AdditionalRootDirs));
-			BundleInfo.ContentPaths.NonUFSShaderLibPaths.Append(MoveTemp(ResultInfo.NonUFSShaderLibPaths));
 			BundleInfo.ContentPaths.ProjectName = ResultInfo.ProjectName;
 		}
 
diff --git a/Engine/Plugins/Experimental/DefaultInstallBundleManager/Source/Private/InstallBundleSourceBulk.cpp b/Engine/Plugins/Experimental/DefaultInstallBundleManager/Source/Private/InstallBundleSourceBulk.cpp
--- a/Engine/Plugins/Experimental/DefaultInstallBundleManager/Source/Private/InstallBundleSourceBulk.cpp
+++ b/Engine/Plugins/Experimental/DefaultInstallBundleManager/Source/Private/InstallBundleSourceBulk.cpp
@@ -617,20 +617,8 @@ void FInstallBundleSourceBulk::RequestUpdateContent(FRequestUpdateContentBundleC
 		ResultInfo.ContentPaths = *BundleFileList;
 	}
 
-	TArray<FString> NonUFSShaderLibPathsExtensions;
-	GConfig->GetArray(TEXT("InstallBundleManager.Shaderlib"), TEXT("ExtensionList"), NonUFSShaderLibPathsExtensions, GInstallBundleIni);
-
 	for (const FString& Path : ResultInfo.ContentPaths)
 	{
-		for (const FString& Ext : NonUFSShaderLibPathsExtensions)
-		{
-			if (Path.EndsWith(*Ext))
-			{
-				LOG_SOURCE_BULK_OVERRIDE(Context.LogVerbosityOverride, Display, TEXT("Found metallib %s for Bundle %s"), *Path, *Context.BundleName.ToString());
-				ResultInfo.NonUFSShaderLibPaths.Add(FPaths::GetPath(Path));
-			}
-		}
-
 		if (Path.EndsWith(TEXTVIEW(".uondemandtoc")))
 		{
 			TUniquePtr<UE::IoStore::FOnDemandMountArgs> MountArgs = MakeUnique<UE::IoStore::FOnDemandMountArgs>();
diff --git a/Engine/Plugins/Experimental/DefaultInstallBundleManager/Source/Public/DefaultInstallBundleManager.h b/Engine/Plugins/Experimental/DefaultInstallBundleManager/Source/Public/DefaultInstallBundleManager.h
--- a/Engine/Plugins/Experimental/DefaultInstallBundleManager/Source/Public/DefaultInstallBundleManager.h
+++ b/Engine/Plugins/Experimental/DefaultInstallBundleManager/Source/Public/DefaultInstallBundleManager.h
@@ -147,7 +147,6 @@ protected:
 	{
 		TArray<FString> ContentPaths;
 		TArray<FString> AdditionalRootDirs;
-		TSet<FString> NonUFSShaderLibPaths;
 		FString ProjectName;
 		bool bContainsChunks = false;
 	};
diff --git a/Engine/Source/Runtime/InstallBundleManager/Private/InstallBundleTypes.cpp b/Engine/Source/Runtime/InstallBundleManager/Private/InstallBundleTypes.cpp
--- a/Engine/Source/Runtime/InstallBundleManager/Private/InstallBundleTypes.cpp
+++ b/Engine/Source/Runtime/InstallBundleManager/Private/InstallBundleTypes.cpp
@@ -313,6 +313,8 @@ bool FInstallBundleCombinedContentState::GetAnyBundleHasState(EInstallBundleInst
 	return false;
 }
 
+PRAGMA_DISABLE_DEPRECATION_WARNINGS
 FInstallBundleSourceUpdateContentResultInfo::FInstallBundleSourceUpdateContentResultInfo() = default;
 FInstallBundleSourceUpdateContentResultInfo::FInstallBundleSourceUpdateContentResultInfo(FInstallBundleSourceUpdateContentResultInfo&&) = default;
 FInstallBundleSourceUpdateContentResultInfo::~FInstallBundleSourceUpdateContentResultInfo() = default;
+PRAGMA_ENABLE_DEPRECATION_WARNINGS
diff --git a/Engine/Source/Runtime/InstallBundleManager/Public/InstallBundleTypes.h b/Engine/Source/Runtime/InstallBundleManager/Public/InstallBundleTypes.h
--- a/Engine/Source/Runtime/InstallBundleManager/Public/InstallBundleTypes.h
+++ b/Engine/Source/Runtime/InstallBundleManager/Public/InstallBundleTypes.h
@@ -387,7 +387,7 @@ struct FInstallBundleSourceUpdateContentResultInfo
 
 	TArray<FString> ContentPaths;
 	TArray<FString> AdditionalRootDirs;
-	// Support platforms that need shaderlibs in the physical FS
+	UE_DEPRECATED(5.6, "All shader libs are now packaged and avialable via UFS")
 	TSet<FString> NonUFSShaderLibPaths;
 	TArray<TUniquePtr<UE::IoStore::FOnDemandMountArgs>> OnDemandMountArgs;
 	FString ProjectName;
