DefaultInstallBundleManager
- Fix HandleAsyncInitError being called every tick while a retry was in progress.  This would spam listeners with callbacks unnecessarly.  Now the handler calls are made in sync with the actual retries (using an exponential backoff cadence).

For sh-11-21-esp-init-001-error-message-ext

diff --git a/Engine/Plugins/Experimental/DefaultInstallBundleManager/Source/Private/DefaultInstallBundleManager.cpp b/Engine/Plugins/Experimental/DefaultInstallBundleManager/Source/Private/DefaultInstallBundleManager.cpp
--- a/Engine/Plugins/Experimental/DefaultInstallBundleManager/Source/Private/DefaultInstallBundleManager.cpp
+++ b/Engine/Plugins/Experimental/DefaultInstallBundleManager/Source/Private/DefaultInstallBundleManager.cpp
@@ -308,74 +308,97 @@ void FDefaultInstallBundleManager::TickInit()
 {
 	CSV_SCOPED_TIMING_STAT(InstallBundleManager, InstallBundleManager_TickInit);
 
+	auto HandleInitStopped = [this]()
+	{
+		LOG_INSTALL_BUNDLE_MAN(Error, TEXT("Initialization Failed - %s"), LexToString(InitResult));
+		InitState = EInstallBundleManagerInitState::Failed;
+		InitCompleteDelegate.Broadcast(InitResult);
+	};
+
+	auto GetConfigInitRetryTimeSecondsMin = []() -> double
+	{
+		double InitRetryTimeSecondsMin = 0.1;
+		GConfig->GetDouble(TEXT("InstallBundleManager.MiscSettings"), TEXT("InitRetryTimeSecondsMin"), InitRetryTimeSecondsMin, GInstallBundleIni);
+		if (InitRetryTimeSecondsMin <= 0)
+		{
+			ensureMsgf(false, TEXT("InitRetryTimeSecondsMin must be greater than 0!"));
+			InitRetryTimeSecondsMin = 0.1;
+		}
+		return InitRetryTimeSecondsMin;
+	};
+
+	auto GetConfigInitRetryTimeSecondsMax = []() -> double
+	{
+		double InitRetryTimeSecondsMax = 5.0;
+		GConfig->GetDouble(TEXT("InstallBundleManager.MiscSettings"), TEXT("InitRetryTimeSecondsMax"), InitRetryTimeSecondsMax, GInstallBundleIni);
+		return InitRetryTimeSecondsMax;
+	};
+
 	while (InitState == EInstallBundleManagerInitState::NotInitialized && InitStepResult == EAsyncInitStepResult::Done)
 	{
 		if (bUnrecoverableInitError)
 		{
-			LOG_INSTALL_BUNDLE_MAN(Error, TEXT("Initialization Failed - %s"), LexToString(InitResult));
-			InitState = EInstallBundleManagerInitState::Failed;
-			InitCompleteDelegate.Broadcast(InitResult);
+			HandleInitStopped();
 			break;
 		}
 
 		if (InitResult != EInstallBundleManagerInitResult::OK)
 		{
-			EInstallBundleManagerInitErrorHandlerResult HandlerResult = HandleAsyncInitError(InitResult);
-
 			if (LastInitStep != InitStep)
 			{
-				// Only fire init analytic for failures the first time we retry
-				AsyncInit_FireInitAnlaytic(HandlerResult != EInstallBundleManagerInitErrorHandlerResult::StopInitialization);
 				LastInitStep = InitStep;
+
+				// Reset retry timer for first retry
+				InitRetryTimeDeltaSeconds = 0;
 			}
 
-			if (HandlerResult == EInstallBundleManagerInitErrorHandlerResult::StopInitialization)
+			// Attempt to try the step again
+			const double CurrentTime = FPlatformTime::Seconds();
+			if (LastInitRetryTimeSeconds <= 0)
 			{
-				LOG_INSTALL_BUNDLE_MAN(Error, TEXT("Initialization Failed - %s"), LexToString(InitResult));
-				InitState = EInstallBundleManagerInitState::Failed;
-				InitCompleteDelegate.Broadcast(InitResult);
-				break;
+				LastInitRetryTimeSeconds = CurrentTime;
 			}
-			else // Try the step again
+
+			if (CurrentTime < LastInitRetryTimeSeconds + InitRetryTimeDeltaSeconds)
 			{
-				double InitRetryTimeSecondsMin = 0.1;
-				GConfig->GetDouble(TEXT("InstallBundleManager.MiscSettings"), TEXT("InitRetryTimeSecondsMin"), InitRetryTimeSecondsMin, GInstallBundleIni);
-				if (InitRetryTimeSecondsMin <= 0)
-				{
-					ensureMsgf(false, TEXT("InitRetryTimeSecondsMin must be greater than 0!"));
-					InitRetryTimeSecondsMin = 0.1;
-				}
+				break;
+			}
 
-				double InitRetryTimeSecondsMax = 5.0;
-				GConfig->GetDouble(TEXT("InstallBundleManager.MiscSettings"), TEXT("InitRetryTimeSecondsMax"), InitRetryTimeSecondsMax, GInstallBundleIni);
-				if (InitRetryTimeSecondsMax < InitRetryTimeSecondsMin)
-				{
-					ensureMsgf(false, TEXT("InitRetryTimeSecondsMax must be greater than or equal to InitRetryTimeSecondsMin!"));
-					InitRetryTimeSecondsMax = InitRetryTimeSecondsMin;
-				}
+			// call error handler before allowing the retry
+			const EInstallBundleManagerInitErrorHandlerResult HandlerResult = HandleAsyncInitError(InitResult);
 
-				double CurrentTime = FPlatformTime::Seconds();
-				if (LastInitRetryTimeSeconds == 0)
-				{
-					LastInitRetryTimeSeconds = CurrentTime;
-				}
+			if (InitRetryTimeDeltaSeconds <= 0)
+			{
+				// Only fire init analytic for failures the first time we retry
+				AsyncInit_FireInitAnlaytic(HandlerResult != EInstallBundleManagerInitErrorHandlerResult::StopInitialization);
+			}
 
-				if (InitRetryTimeDeltaSeconds == 0)
-				{
-					InitRetryTimeDeltaSeconds = InitRetryTimeSecondsMin;
-				}
+			if (HandlerResult == EInstallBundleManagerInitErrorHandlerResult::StopInitialization)
+			{
+				HandleInitStopped();
+				break;
+			}
 
-				if (CurrentTime < LastInitRetryTimeSeconds + InitRetryTimeDeltaSeconds)
-				{
-					break;
-				}
+			if (InitStep != EAsyncInitStep::None) // Don't spam this for unrecoverable errors
+			{
+				LOG_INSTALL_BUNDLE_MAN(Warning, TEXT("Retrying initialization after %s - waited %f seconds"), LexToString(InitResult), InitRetryTimeDeltaSeconds);
+			}
 
-				if (InitStep != EAsyncInitStep::None) // Don't spam this for unrecoverable errors
+			const double ConfigInitRetryTimeSecondsMin = GetConfigInitRetryTimeSecondsMin();
+			if (InitRetryTimeDeltaSeconds <= 0)
+			{
+				InitRetryTimeDeltaSeconds = ConfigInitRetryTimeSecondsMin;
+			}
+			else
+			{
+				double ConfigInitRetryTimeSecondsMax = GetConfigInitRetryTimeSecondsMax();
+				if (ConfigInitRetryTimeSecondsMax < ConfigInitRetryTimeSecondsMin)
 				{
-					LOG_INSTALL_BUNDLE_MAN(Warning, TEXT("Retrying initialization after %s - waited %f seconds"), LexToString(InitResult), InitRetryTimeDeltaSeconds);
+					ensureMsgf(false, TEXT("InitRetryTimeSecondsMax must be greater than or equal to InitRetryTimeSecondsMin!"));
+					ConfigInitRetryTimeSecondsMax = ConfigInitRetryTimeSecondsMin;
 				}
 
-				InitRetryTimeDeltaSeconds = FMath::Min(InitRetryTimeDeltaSeconds * 2.0, InitRetryTimeSecondsMax);
+				InitRetryTimeDeltaSeconds = FMath::Min(InitRetryTimeDeltaSeconds * 2.0, ConfigInitRetryTimeSecondsMax);
 			}
 		}
 
