Clean up install bundle manager OnDemandIoStore module usage.

diff --git a/Engine/Plugins/Experimental/DefaultInstallBundleManager/Source/Private/DefaultInstallBundleManager.cpp b/Engine/Plugins/Experimental/DefaultInstallBundleManager/Source/Private/DefaultInstallBundleManager.cpp
--- a/Engine/Plugins/Experimental/DefaultInstallBundleManager/Source/Private/DefaultInstallBundleManager.cpp
+++ b/Engine/Plugins/Experimental/DefaultInstallBundleManager/Source/Private/DefaultInstallBundleManager.cpp
@@ -1393,8 +1393,8 @@ void FDefaultInstallBundleManager::UpdateBundleSources(FContentReleaseRequestRef
 	// Only unmount on demand tocs if this is a full release
 	if (BundleInfo.bMountedOnDemandTocs)
 	{
-		check(IOStoreOnDemandModule);
-		const FIoStatus UnmountStatus = UE::IoStore::GetOnDemandIoStore().Unmount(Request->BundleName.ToString());
+		check(OnDemandIoStore);
+		const FIoStatus UnmountStatus = OnDemandIoStore->Unmount(Request->BundleName.ToString());
 		if (!UnmountStatus.IsOk())
 		{
 			LOG_INSTALL_BUNDLE_MAN_OVERRIDE(Request->LogVerbosityOverride, Error, TEXT("Error unmounting ondemand toc for bundle '%s' : %s"), *BundleInfo.BundleNameString, *UnmountStatus.ToString());
@@ -1653,10 +1653,10 @@ void FDefaultInstallBundleManager::MountPaks(FContentRequestRef Request)
 		return;
 	}
 
-	if (!IOStoreOnDemandModule)
+	if (!OnDemandIoStore)
 	{
-		IOStoreOnDemandModule = FModuleManager::Get().GetModulePtr<UE::IoStore::FIoStoreOnDemandModule>(TEXT("IoStoreOnDemand"));
-		if (!IOStoreOnDemandModule)
+		OnDemandIoStore = UE::IoStore::TryGetOnDemandIoStore();
+		if (!OnDemandIoStore)
 		{
 			Request->Result = EInstallBundleResult::InitializationError;
 			Request->OptionalErrorCode = TEXT("IoStoreOnDemand_Not_Found");
diff --git a/Engine/Plugins/Experimental/DefaultInstallBundleManager/Source/Public/DefaultInstallBundleManager.h b/Engine/Plugins/Experimental/DefaultInstallBundleManager/Source/Public/DefaultInstallBundleManager.h
--- a/Engine/Plugins/Experimental/DefaultInstallBundleManager/Source/Public/DefaultInstallBundleManager.h
+++ b/Engine/Plugins/Experimental/DefaultInstallBundleManager/Source/Public/DefaultInstallBundleManager.h
@@ -16,7 +16,7 @@
 enum class EPatchCheckResult : uint8;
 class IAnalyticsProviderET;
 class IInstallBundleSource;
-namespace UE::IoStore { class FIoStoreOnDemandModule; }
+namespace UE::IoStore { class IOnDemandIoStore; }
 
 class DEFAULTINSTALLBUNDLEMANAGER_API FDefaultInstallBundleManager : public IInstallBundleManager
 {
@@ -734,7 +734,7 @@ protected:
 	FTSTicker::FDelegateHandle TickHandle;
 	FDelegateHandle AsyncInit_PatchCheckHandle;
 	FDelegateHandle PatchCheckHandle;
-	UE::IoStore::FIoStoreOnDemandModule* IOStoreOnDemandModule = nullptr;
+	UE::IoStore::IOnDemandIoStore* OnDemandIoStore = nullptr;
 	FInstallBundleSourceFactoryFunction InstallBundleSourceFactory;
 
 	TMap<FName, FBundleInfo> BundleInfoMap;
