FInstallBundleSourcePlatformChunkInstall - Fix bundle name not being set when returning bundle progress.

diff --git a/Engine/Plugins/Experimental/DefaultInstallBundleManager/Source/Private/InstallBundleSourcePlatformChunkInstall.cpp b/Engine/Plugins/Experimental/DefaultInstallBundleManager/Source/Private/InstallBundleSourcePlatformChunkInstall.cpp
--- a/Engine/Plugins/Experimental/DefaultInstallBundleManager/Source/Private/InstallBundleSourcePlatformChunkInstall.cpp
+++ b/Engine/Plugins/Experimental/DefaultInstallBundleManager/Source/Private/InstallBundleSourcePlatformChunkInstall.cpp
@@ -635,6 +635,7 @@ TOptional<FInstallBundleSourceProgress> FInstallBundleSourcePlatformChunkInstall
 		float ProgressPercent = PlatformChunkInstall->GetNamedChunkProgress(BundleInfo->NamedChunk, EChunkProgressReportingType::PercentageComplete);
 
 		Status.Emplace();
+		Status->BundleName = BundleName;
 		Status->Install_Percent = (ProgressPercent / 100.0f);
 	}
 
