FDefaultInstallBundleManager
- Allow bundle sources not to return progress for a bundle.  This fixes cases where there is trivial or no progress from a source causing inflation of overall progress.

diff --git a/Engine/Plugins/Experimental/DefaultInstallBundleManager/Source/Private/DefaultInstallBundleManager.cpp b/Engine/Plugins/Experimental/DefaultInstallBundleManager/Source/Private/DefaultInstallBundleManager.cpp
--- a/Engine/Plugins/Experimental/DefaultInstallBundleManager/Source/Private/DefaultInstallBundleManager.cpp
+++ b/Engine/Plugins/Experimental/DefaultInstallBundleManager/Source/Private/DefaultInstallBundleManager.cpp
@@ -4083,7 +4083,9 @@ TOptional<FInstallBundleProgress> FDefaultInstallBundleManager::GetBundleProgres
 
 	const FBundleInfo* BundleInfo = BundleInfoMap.Find(BundleName);
 	if (BundleInfo == nullptr)
+	{
 		return CombinedStatus;
+	}
 
 	// Find the request
 	FContentRequestPtr Request;
@@ -4124,7 +4126,9 @@ TOptional<FInstallBundleProgress> FDefaultInstallBundleManager::GetBundleProgres
 	CombinedStatus->Status = EInstallBundleStatus::Updating;
 
 	if (!Request->Steps.IsValidIndex(Request->iStep))
+	{
 		return CombinedStatus;
+	}
 
 	EContentRequestState State = Request->Steps[Request->iStep];
 
@@ -4158,13 +4162,19 @@ TOptional<FInstallBundleProgress> FDefaultInstallBundleManager::GetBundleProgres
 			CombinedStatus->Install_Percent += SourceWeight * Pair.Value.Install_Percent;
 		}
 
-		check(TotalWeight >= 0.0f);
 		if (TotalWeight > 0.0f)
 		{
 			CombinedStatus->BackgroundDownload_Percent /= TotalWeight;
 			CombinedStatus->InstallOnly_Percent /= TotalWeight;
 			CombinedStatus->Install_Percent /= TotalWeight;
 		}
+		else
+		{
+			// No progress returned from any source, assume everything is done
+			CombinedStatus->BackgroundDownload_Percent = 1.0f;
+			CombinedStatus->InstallOnly_Percent = 1.0f;
+			CombinedStatus->Install_Percent = 1.0f;
+		}
 	}
 	else if (State > EContentRequestState::UpdatingBundleSources && Request->Result == EInstallBundleResult::OK)
 	{
