Install Bundle Manager
- Add bMountSoftReferences flag to mount iostore soft refs
- Use new FCoreDelegates::MountContainers to pass mount options and get errors.

diff --git a/Engine/Plugins/Experimental/DefaultInstallBundleManager/Source/Private/DefaultInstallBundleManager.cpp b/Engine/Plugins/Experimental/DefaultInstallBundleManager/Source/Private/DefaultInstallBundleManager.cpp
--- a/Engine/Plugins/Experimental/DefaultInstallBundleManager/Source/Private/DefaultInstallBundleManager.cpp
+++ b/Engine/Plugins/Experimental/DefaultInstallBundleManager/Source/Private/DefaultInstallBundleManager.cpp
@@ -6,14 +6,11 @@
 #include "InstallBundleSourceInterface.h"
 #include "InstallBundleManagerModule.h"
 
-
 #include "HAL/IConsoleManager.h"
 #include "Misc/Base64.h"
 #include "Misc/ConfigContext.h"
 #include "Misc/CommandLine.h"
 #include "IO/IoStoreOnDemand.h"
-#include "IPlatformFilePak.h"
-#include "HAL/PlatformFileManager.h"
 #include "PatchCheck.h"
 #include "Templates/Greater.h"
 #include "ShaderPipelineCache.h"
@@ -1312,10 +1309,10 @@ void FDefaultInstallBundleManager::UpdateBundleSourceComplete(TSharedRef<IInstal
 				// lets save a little memory
 				if (ContentPath.EndsWith(TEXTVIEW(".pak")))
 				{
-					BundleInfo.ContentPaths.ContentPaths.Add(MoveTemp(ContentPath));
+					BundleInfo.ContentPaths.ContentPaths.Emplace(MoveTemp(ContentPath), ResultInfo.MountOptions);
 				}
 			}
-			
+
 			Request->OnDemandMountArgs.Append(MoveTemp(ResultInfo.OnDemandMountArgs));
 
 			BundleInfo.ContentPaths.AdditionalRootDirs.Append(MoveTemp(ResultInfo.AdditionalRootDirs));
@@ -1329,6 +1326,12 @@ void FDefaultInstallBundleManager::UpdateBundleSourceComplete(TSharedRef<IInstal
 	}
 	Request->SourceRequestResults.Empty();
 
+	// Sort in descending order for determinism
+	BundleInfo.ContentPaths.ContentPaths.Sort([](const TPair<FString, FPakMountOptions>& A, const TPair<FString, FPakMountOptions>& B)
+	{
+		return A.Key > B.Key;
+	});
+
 	ensureMsgf(!Request->OnDemandMountArgs.IsEmpty() == BundleInfo.bContainsIoStoreOnDemandTocs,
 		TEXT("OnDemandMountArgs should only be provided by bundle sources that set bContainsIoStoreOnDemandTocs"));
 
@@ -1336,8 +1339,9 @@ void FDefaultInstallBundleManager::UpdateBundleSourceComplete(TSharedRef<IInstal
 	// This is a corner case but as far as I know there is no other situation that would allow for an empty bundle that was not chunked.
 	// If such a case were to arise, then bContainsChunks would need to be determined by each bundle source individually.
 	BundleInfo.ContentPaths.bContainsChunks = BundleInfo.ContentPaths.ContentPaths.IsEmpty();
-	for (const FString& ContentPath : BundleInfo.ContentPaths.ContentPaths)
+	for (const TPair<FString, FPakMountOptions>& Pair : BundleInfo.ContentPaths.ContentPaths)
 	{
+		const FString& ContentPath = Pair.Key;
 		BundleInfo.ContentPaths.bContainsChunks = 
 			ContentPath.EndsWith(TEXTVIEW(".pak")) &&
 			FPlatformMisc::GetPakchunkIndexFromPakFile(ContentPath) != INDEX_NONE;
@@ -1754,75 +1758,39 @@ void FDefaultInstallBundleManager::MountPaks(FContentRequestRef Request)
 	}
 }
 
-bool FDefaultInstallBundleManager::MountPaksInList(TArrayView<FString> Paths, ELogVerbosity::Type LogVerbosityOverride)
+bool FDefaultInstallBundleManager::MountPaksInList(TArrayView<TPair<FString, FPakMountOptions>> Paths, ELogVerbosity::Type LogVerbosityOverride)
 {
-	FPakPlatformFile* PakPlatformFile = static_cast<FPakPlatformFile*>(FPlatformFileManager::Get().FindPlatformFile(TEXT("PakFile")));
-
-	// Sort in descending order.
-	Paths.Sort(TGreater<FString>());
-
-	// Find already mounted paks
-	TSet<FString> MountedPaks;
-	if (PakPlatformFile)
+	if (Paths.Num() > 0 && !FCoreDelegates::MountPaksEx.IsBound())
 	{
-		PakPlatformFile->GetMountedPakFilenames(MountedPaks);
+		ensureMsgf(false, TEXT("Pak files have not been correctly initalized. Use -UsePaks on the cmdline if you are using the UnrealEditor.exe"));
+		return false;
 	}
 
-	bool bMountedPaks = false;
-	for (const FString& File : Paths)
-	{
-		if (!File.EndsWith(TEXTVIEW(".pak")))
-		{
-			continue;
-		}
+	TArray<UE::FMountPaksExArgs> MountArgs;
+	MountArgs.Reserve(Paths.Num());
 
-		if (MountedPaks.Contains(File))
-		{
-			LOG_INSTALL_BUNDLE_MAN_OVERRIDE(LogVerbosityOverride, Warning, TEXT("Pak file: %s already mounted, skipping. \n"), *File);
-			continue;
-		}
-
-		if (!PakPlatformFile)
-		{
-			ensureMsgf(false, TEXT("Pak files have not been correctly initalized. Use -UsePaks on the cmdline if you are using the UnrealEditor.exe"));
-			return false; // if FCoreDelegates::MountPak is unbound there is a major issue.
-		}
-
-		LOG_INSTALL_BUNDLE_MAN_OVERRIDE(LogVerbosityOverride, Verbose, TEXT("Mounting pak file: %s \n"), *File);
-
-		if (FCoreDelegates::MountPak.Execute(File, INDEX_NONE)) //May fail on encrypted Paks
+	for (const TPair<FString, FPakMountOptions>& Pair : Paths)
+	{
+		MountArgs.Emplace(UE::FMountPaksExArgs
 		{
-			bMountedPaks = true;
-		}
+			.PakFilePath = *Pair.Key,
+			.MountOptions = Pair.Value
+		});
 	}
 
-	return bMountedPaks;
-}
+	//May not mount encrypted Paks
+	const bool bMountedPaks = FCoreDelegates::MountPaksEx.Execute(MountArgs);
 
-bool FDefaultInstallBundleManager::UnmountPaksInList(TArrayView<FString> Paths, ELogVerbosity::Type LogVerbosityOverride)
-{
-	if (!FCoreDelegates::OnUnmountPak.IsBound())
+	for (const UE::FMountPaksExArgs& Args : MountArgs)
 	{
-		return false;
-	}
-
-	bool bUnmountedPaks = false;
-	for (const FString& File : Paths)
-	{
-		if (!File.EndsWith(TEXTVIEW(".pak")))
-		{
-			continue;
-		}
-
-		LOG_INSTALL_BUNDLE_MAN_OVERRIDE(LogVerbosityOverride, Display, TEXT("Unmounting pak file: %s \n"), *File);
-
-		if (FCoreDelegates::OnUnmountPak.Execute(File))
+		if (const FString* Error = Args.Result.TryGetError())
 		{
-			bUnmountedPaks = true;
+			// TODO: fail request now that we can actually get errors from mounting
+			LOG_INSTALL_BUNDLE_MAN_OVERRIDE(LogVerbosityOverride, Error, TEXT("Failed to mount pak %s, error: %s"), Args.PakFilePath, **Error);
 		}
 	}
 
-	return bUnmountedPaks;
+	return bMountedPaks;
 }
 
 void FDefaultInstallBundleManager::UnmountPaks(FContentReleaseRequestRef Request)
@@ -1837,7 +1805,14 @@ void FDefaultInstallBundleManager::UnmountPaks(FContentReleaseRequestRef Request
 
 	LOG_INSTALL_BUNDLE_MAN_OVERRIDE(Request->LogVerbosityOverride, Display, TEXT("Unmounting Paks for Release Request %s"), *BundleInfo.BundleNameString);
 
-	UnmountPaksInList(BundleInfo.ContentPaths.ContentPaths, Request->LogVerbosityOverride);
+	if (FCoreDelegates::OnUnmountPak.IsBound())
+	{
+		for (const TPair<FString, FPakMountOptions>& Pair : BundleInfo.ContentPaths.ContentPaths)
+		{
+			LOG_INSTALL_BUNDLE_MAN_OVERRIDE(Request->LogVerbosityOverride, Display, TEXT("Unmounting pak file: %s \n"), *Pair.Key);
+			FCoreDelegates::OnUnmountPak.Execute(Pair.Key);
+		}
+	}
 
 	OnPaksUnmountedInternal(Request, BundleInfo);
 
@@ -1851,18 +1826,9 @@ bool FDefaultInstallBundleManager::AllowIoStoreOnDemandMount(FContentRequestRef
 	return true;
 }
 
-TArray<FString> FDefaultInstallBundleManager::GetPakMountList(FContentRequestRef Request, const FBundleInfo& BundleInfo)
+TArray<TPair<FString, FPakMountOptions>> FDefaultInstallBundleManager::GetPakMountList(FContentRequestRef Request, const FBundleInfo& BundleInfo)
 {
-	TArray<FString> PakList;
-	for (const FString& Path : BundleInfo.ContentPaths.ContentPaths)
-	{
-		if (Path.EndsWith(TEXTVIEW(".pak")))
-		{
-			PakList.Add(Path);
-		}
-	}
-
-	return PakList;
+	return BundleInfo.ContentPaths.ContentPaths;
 }
 
 void FDefaultInstallBundleManager::WaitForShaderCache(FContentRequestRef Request)
diff --git a/Engine/Plugins/Experimental/DefaultInstallBundleManager/Source/Private/InstallBundleSourceBulk.cpp b/Engine/Plugins/Experimental/DefaultInstallBundleManager/Source/Private/InstallBundleSourceBulk.cpp
--- a/Engine/Plugins/Experimental/DefaultInstallBundleManager/Source/Private/InstallBundleSourceBulk.cpp
+++ b/Engine/Plugins/Experimental/DefaultInstallBundleManager/Source/Private/InstallBundleSourceBulk.cpp
@@ -606,7 +606,12 @@ void FInstallBundleSourceBulk::GetContentState(TArrayView<const FName> BundleNam
 
 void FInstallBundleSourceBulk::RequestUpdateContent(FRequestUpdateContentBundleContext Context)
 {
-	LOG_SOURCE_BULK_OVERRIDE(Context.LogVerbosityOverride, Display, TEXT("Requesting Bundle %s"), *Context.BundleName.ToString());
+	FNameBuilder BundleNameBuilder(Context.BundleName);
+
+	LOG_SOURCE_BULK_OVERRIDE(Context.LogVerbosityOverride, Display, TEXT("Requesting Bundle %s"), BundleNameBuilder.ToString());
+
+	InstallBundleUtil::FConfigMountOptions OnDemandMountOptions;
+	InstallBundleUtil::GetMountOptionsFromConfig(BundleNameBuilder, OnDemandMountOptions);
 
 	FInstallBundleSourceUpdateContentResultInfo ResultInfo;
 	ResultInfo.BundleName = Context.BundleName;
@@ -628,9 +633,19 @@ void FInstallBundleSourceBulk::RequestUpdateContent(FRequestUpdateContentBundleC
 			MountArgs->FilePath = Path;
 			MountArgs->Options = UE::IoStore::EOnDemandMountOptions::InstallOnDemand;
 
+			if (OnDemandMountOptions.bWithSoftReferences)
+			{
+				MountArgs->Options |= UE::IoStore::EOnDemandMountOptions::WithSoftReferences;
+			}
+
 			ResultInfo.OnDemandMountArgs.Add(MoveTemp(MountArgs));
 		}
 	}
+
+	if (OnDemandMountOptions.bWithSoftReferences)
+	{
+		ResultInfo.MountOptions.MountFlags |= FPakMountOptions::EMountFlags::WithSoftReferences;
+	}
 	
 	ResultInfo.ProjectName = FApp::GetProjectName();
 	Context.CompleteCallback.ExecuteIfBound(AsShared(), MoveTemp(ResultInfo));
diff --git a/Engine/Plugins/Experimental/DefaultInstallBundleManager/Source/Public/DefaultInstallBundleManager.h b/Engine/Plugins/Experimental/DefaultInstallBundleManager/Source/Public/DefaultInstallBundleManager.h
--- a/Engine/Plugins/Experimental/DefaultInstallBundleManager/Source/Public/DefaultInstallBundleManager.h
+++ b/Engine/Plugins/Experimental/DefaultInstallBundleManager/Source/Public/DefaultInstallBundleManager.h
@@ -17,6 +17,7 @@ enum class EPatchCheckResult : uint8;
 class IAnalyticsProviderET;
 class IInstallBundleSource;
 namespace UE::IoStore { class IOnDemandIoStore; }
+struct FPakMountOptions;
 
 class DEFAULTINSTALLBUNDLEMANAGER_API FDefaultInstallBundleManager : public IInstallBundleManager
 {
@@ -145,7 +146,7 @@ protected:
 
 	struct FBundleContentPaths
 	{
-		TArray<FString> ContentPaths;
+		TArray<TPair<FString, FPakMountOptions>> ContentPaths;
 		TArray<FString> AdditionalRootDirs;
 		FString ProjectName;
 		bool bContainsChunks = false;
@@ -557,15 +558,13 @@ protected:
 
 	void MountPaks(FContentRequestRef Request);
 
-	static bool MountPaksInList(TArrayView<FString> Paths, ELogVerbosity::Type LogVerbosityOverride);
-
-	static bool UnmountPaksInList(TArrayView<FString> Paths, ELogVerbosity::Type LogVerbosityOverride);
+	static bool MountPaksInList(TArrayView<TPair<FString, FPakMountOptions>> Paths, ELogVerbosity::Type LogVerbosityOverride);
 
 	void UnmountPaks(FContentReleaseRequestRef Request);
 
 	virtual bool AllowIoStoreOnDemandMount(FContentRequestRef Request, const FBundleInfo& BundleInfo);
 
-	virtual TArray<FString> GetPakMountList(FContentRequestRef Request, const FBundleInfo& BundleInfo);
+	virtual TArray<TPair<FString, FPakMountOptions>> GetPakMountList(FContentRequestRef Request, const FBundleInfo& BundleInfo);
 
 	virtual void OnPaksMountedInternal(FContentRequestRef Request, FBundleInfo& BundleInfo) {}
 	virtual void OnPaksUnmountedInternal(FContentReleaseRequestRef Request, FBundleInfo& BundleInfo) {};
diff --git a/Engine/Source/Runtime/InstallBundleManager/Private/InstallBundleUtils.cpp b/Engine/Source/Runtime/InstallBundleManager/Private/InstallBundleUtils.cpp
--- a/Engine/Source/Runtime/InstallBundleManager/Private/InstallBundleUtils.cpp
+++ b/Engine/Source/Runtime/InstallBundleManager/Private/InstallBundleUtils.cpp
@@ -156,6 +156,22 @@ namespace InstallBundleUtil
 		return false;
 	}
 
+	bool GetMountOptionsFromConfig(const FStringView BundleName, FConfigMountOptions& OutMountOptions)
+	{
+		const FConfigFile* InstallBundleConfig = GConfig->FindConfigFile(GInstallBundleIni);
+		if (InstallBundleConfig)
+		{
+			const FString SectionName = InstallBundleUtil::GetInstallBundleSectionPrefix() + BundleName;
+			if (InstallBundleConfig->DoesSectionExist(*SectionName))
+			{
+				OutMountOptions.bWithSoftReferences = false;
+				InstallBundleConfig->GetBool(*SectionName, TEXT("bMountSoftReferences"), OutMountOptions.bWithSoftReferences);
+				return true;
+			}
+		}
+		return false;
+	}
+
 	bool AllInstallBundlePredicate(const FConfigFile& InstallBundleConfig, const FString& Section)
 	{ 
 		return true; 
diff --git a/Engine/Source/Runtime/InstallBundleManager/Public/InstallBundleTypes.h b/Engine/Source/Runtime/InstallBundleManager/Public/InstallBundleTypes.h
--- a/Engine/Source/Runtime/InstallBundleManager/Public/InstallBundleTypes.h
+++ b/Engine/Source/Runtime/InstallBundleManager/Public/InstallBundleTypes.h
@@ -9,6 +9,7 @@
 #include "Containers/UnrealString.h"
 #include "CoreMinimal.h"
 #include "Delegates/Delegate.h"
+#include "GenericPlatform/GenericPlatformFile.h"
 #include "HAL/Platform.h"
 #include "Internationalization/Text.h"
 #include "Misc/DateTime.h"
@@ -391,6 +392,7 @@ struct FInstallBundleSourceUpdateContentResultInfo
 	UE_DEPRECATED(5.6, "All shader libs are now packaged and avialable via UFS")
 	TSet<FString> NonUFSShaderLibPaths;
 	TArray<TUniquePtr<UE::IoStore::FOnDemandMountArgs>> OnDemandMountArgs;
+	FPakMountOptions MountOptions;
 	FString ProjectName;
 
 	uint64 CurrentInstallSize = 0;
diff --git a/Engine/Source/Runtime/InstallBundleManager/Public/InstallBundleUtils.h b/Engine/Source/Runtime/InstallBundleManager/Public/InstallBundleUtils.h
--- a/Engine/Source/Runtime/InstallBundleManager/Public/InstallBundleUtils.h
+++ b/Engine/Source/Runtime/InstallBundleManager/Public/InstallBundleUtils.h
@@ -44,6 +44,13 @@ namespace InstallBundleUtil
 	// returns true if the given BundleName exists in the InstallBundle.ini config
 	INSTALLBUNDLEMANAGER_API bool HasInstallBundleInConfig(const FString& BundleName);
 
+	struct FConfigMountOptions
+	{
+		bool bWithSoftReferences = false; 
+	};
+
+	INSTALLBUNDLEMANAGER_API bool GetMountOptionsFromConfig(const FStringView BundleName, FConfigMountOptions& OutMountOptions);
+
 	// Returns true
 	INSTALLBUNDLEMANAGER_API bool AllInstallBundlePredicate(const FConfigFile& InstallBundleConfig, const FString& Section);
 
