Install Bundle Manger
- Propogate errors that happen during mount, still need FError to do this properly.

diff --git a/Engine/Plugins/Experimental/DefaultInstallBundleManager/Source/Private/DefaultInstallBundleManager.cpp b/Engine/Plugins/Experimental/DefaultInstallBundleManager/Source/Private/DefaultInstallBundleManager.cpp
--- a/Engine/Plugins/Experimental/DefaultInstallBundleManager/Source/Private/DefaultInstallBundleManager.cpp
+++ b/Engine/Plugins/Experimental/DefaultInstallBundleManager/Source/Private/DefaultInstallBundleManager.cpp
@@ -1583,12 +1583,12 @@ void FDefaultInstallBundleManager::MountPaks(FContentRequestRef Request)
 
 	auto MountPaksWorkFunc =
 		[ContentPaths = GetPakMountList(Request, BundleInfo)
-		, LogVerbosityOverride = Request->LogVerbosityOverride]() mutable
+		, LogVerbosityOverride = Request->LogVerbosityOverride](TValueOrError<void, FString>& OutResult) mutable -> bool
 	{
-		return MountPaksInList(ContentPaths, LogVerbosityOverride);
+		return MountPaksInList(ContentPaths, OutResult, LogVerbosityOverride);
 	};
 
-	auto OnPakMountCompleteFunc = [this, WeakThis = AsWeak(), Request](bool bMountedPaks) mutable
+	auto OnPakMountCompleteFunc = [this, WeakThis = AsWeak(), Request](bool bMountedPaks, const TValueOrError<void, FString>& MountResult) mutable
 	{
 		TSharedPtr<IInstallBundleManager> PinnedThis = WeakThis.Pin();
 		if (!PinnedThis)
@@ -1603,6 +1603,19 @@ void FDefaultInstallBundleManager::MountPaks(FContentRequestRef Request)
 			FCoreDelegates::OnOptimizeMemoryUsageForMountedPaks.Execute();
 		}
 
+		if (!MountResult.HasValue())
+		{
+			// TODO: Improve this with FError - currently the error code is lost
+			Request->Result = EInstallBundleResult::InstallError;
+			// Request->OptionalErrorCode = MoveTemp(ErrorCode);
+			Request->OptionalErrorText = FText::AsCultureInvariant(MountResult.GetError()); // Not currently localized 
+			Request->StepResult = EContentRequestStepResult::Done;
+
+			StatsEnd(Request->BundleName, EContentRequestState::Mounting);
+			TRACE_BOOKMARK(TEXT("Finished Mount Bundle %s"), *Request->BundleName.ToString());
+			return;
+		}
+
 		for (const FString& RootDir : BundleInfoLambda.ContentPaths.AdditionalRootDirs)
 		{
 			FPlatformMisc::AddAdditionalRootDirectory(RootDir);
@@ -1626,8 +1639,9 @@ void FDefaultInstallBundleManager::MountPaks(FContentRequestRef Request)
 	{
 		if (!bIsAsyncMount)
 		{
-			const bool bMountedPaks = MountPaksWorkFunc();
-			OnPakMountCompleteFunc(bMountedPaks);
+			TValueOrError<void, FString> MountResult = MakeValue();
+			const bool bMountedPaks = MountPaksWorkFunc(MountResult);
+			OnPakMountCompleteFunc(bMountedPaks, MountResult);
 		}
 		else
 		{
@@ -1635,10 +1649,14 @@ void FDefaultInstallBundleManager::MountPaks(FContentRequestRef Request)
 				[MountPaksWorkFunc = MoveTemp(MountPaksWorkFunc)
 				, OnPakMountCompleteFunc = MoveTemp(OnPakMountCompleteFunc)]() mutable
 			{
-				const bool bMountedPaks = MountPaksWorkFunc();
-				ExecuteOnGameThread(UE_SOURCE_LOCATION, [OnPakMountCompleteFunc = MoveTemp(OnPakMountCompleteFunc), bMountedPaks]() mutable
+				TValueOrError<void, FString> MountResult = MakeValue();
+				const bool bMountedPaks = MountPaksWorkFunc(MountResult);
+				ExecuteOnGameThread(UE_SOURCE_LOCATION, 
+					[OnPakMountCompleteFunc = MoveTemp(OnPakMountCompleteFunc)
+					, MountResult = MoveTemp(MountResult)
+					, bMountedPaks]() mutable
 				{
-					OnPakMountCompleteFunc(bMountedPaks);
+					OnPakMountCompleteFunc(bMountedPaks, MountResult);
 				});
 			});
 		}
@@ -1758,7 +1776,7 @@ void FDefaultInstallBundleManager::MountPaks(FContentRequestRef Request)
 	}
 }
 
-bool FDefaultInstallBundleManager::MountPaksInList(TArrayView<TPair<FString, FPakMountOptions>> Paths, ELogVerbosity::Type LogVerbosityOverride)
+bool FDefaultInstallBundleManager::MountPaksInList(TArrayView<TPair<FString, FPakMountOptions>> Paths, TValueOrError<void, FString>& OutResult, ELogVerbosity::Type LogVerbosityOverride)
 {
 	if (Paths.Num() > 0 && !FCoreDelegates::MountPaksEx.IsBound())
 	{
@@ -1766,6 +1784,8 @@ bool FDefaultInstallBundleManager::MountPaksInList(TArrayView<TPair<FString, FPa
 		return false;
 	}
 
+	OutResult = MakeValue();
+
 	TArray<UE::FMountPaksExArgs> MountArgs;
 	MountArgs.Reserve(Paths.Num());
 
@@ -1781,12 +1801,16 @@ bool FDefaultInstallBundleManager::MountPaksInList(TArrayView<TPair<FString, FPa
 	//May not mount encrypted Paks
 	const bool bMountedPaks = FCoreDelegates::MountPaksEx.Execute(MountArgs);
 
-	for (const UE::FMountPaksExArgs& Args : MountArgs)
+	for (UE::FMountPaksExArgs& Args : MountArgs)
 	{
-		if (const FString* Error = Args.Result.TryGetError())
+		if (Args.Result.HasError())
 		{
-			// TODO: fail request now that we can actually get errors from mounting
-			LOG_INSTALL_BUNDLE_MAN_OVERRIDE(LogVerbosityOverride, Error, TEXT("Failed to mount pak %s, error: %s"), Args.PakFilePath, **Error);
+			LOG_INSTALL_BUNDLE_MAN_OVERRIDE(LogVerbosityOverride, Error, TEXT("Failed to mount pak %s, error: %s"), Args.PakFilePath, *Args.Result.GetError());
+
+			if (!OutResult.HasError())
+			{
+				OutResult = MakeError(Args.Result.StealError());
+			}
 		}
 	}
 
diff --git a/Engine/Plugins/Experimental/DefaultInstallBundleManager/Source/Public/DefaultInstallBundleManager.h b/Engine/Plugins/Experimental/DefaultInstallBundleManager/Source/Public/DefaultInstallBundleManager.h
--- a/Engine/Plugins/Experimental/DefaultInstallBundleManager/Source/Public/DefaultInstallBundleManager.h
+++ b/Engine/Plugins/Experimental/DefaultInstallBundleManager/Source/Public/DefaultInstallBundleManager.h
@@ -558,7 +558,7 @@ protected:
 
 	void MountPaks(FContentRequestRef Request);
 
-	static bool MountPaksInList(TArrayView<TPair<FString, FPakMountOptions>> Paths, ELogVerbosity::Type LogVerbosityOverride);
+	static bool MountPaksInList(TArrayView<TPair<FString, FPakMountOptions>> Paths, TValueOrError<void, FString>& OutResult, ELogVerbosity::Type LogVerbosityOverride);
 
 	void UnmountPaks(FContentReleaseRequestRef Request);
 
