Added error handling to mount pak ex.

diff --git a/Engine/Plugins/Experimental/DefaultInstallBundleManager/Source/Private/DefaultInstallBundleManager.cpp b/Engine/Plugins/Experimental/DefaultInstallBundleManager/Source/Private/DefaultInstallBundleManager.cpp
--- a/Engine/Plugins/Experimental/DefaultInstallBundleManager/Source/Private/DefaultInstallBundleManager.cpp
+++ b/Engine/Plugins/Experimental/DefaultInstallBundleManager/Source/Private/DefaultInstallBundleManager.cpp
@@ -1583,12 +1583,12 @@ void FDefaultInstallBundleManager::MountPaks(FContentRequestRef Request)
 
 	auto MountPaksWorkFunc =
 		[ContentPaths = GetPakMountList(Request, BundleInfo)
-		, LogVerbosityOverride = Request->LogVerbosityOverride](TValueOrError<void, FString>& OutResult) mutable -> bool
+		, LogVerbosityOverride = Request->LogVerbosityOverride](TValueOrError<void, UE::UnifiedError::FError>& OutResult) mutable -> bool
 	{
 		return MountPaksInList(ContentPaths, OutResult, LogVerbosityOverride);
 	};
 
-	auto OnPakMountCompleteFunc = [this, WeakThis = AsWeak(), Request](bool bMountedPaks, const TValueOrError<void, FString>& MountResult) mutable
+	auto OnPakMountCompleteFunc = [this, WeakThis = AsWeak(), Request](bool bMountedPaks, const TValueOrError<void, UE::UnifiedError::FError>& MountResult) mutable
 	{
 		TSharedPtr<IInstallBundleManager> PinnedThis = WeakThis.Pin();
 		if (!PinnedThis)
@@ -1607,8 +1607,7 @@ void FDefaultInstallBundleManager::MountPaks(FContentRequestRef Request)
 		{
 			// TODO: Improve this with FError - currently the error code is lost
 			Request->Result = EInstallBundleResult::InstallError;
-			// Request->OptionalErrorCode = MoveTemp(ErrorCode);
-			Request->OptionalErrorText = FText::AsCultureInvariant(MountResult.GetError()); // Not currently localized 
+			Request->OptionalErrorText = MountResult.GetError().GetErrorMessage(); 
 			Request->StepResult = EContentRequestStepResult::Done;
 
 			StatsEnd(Request->BundleName, EContentRequestState::Mounting);
@@ -1639,7 +1638,7 @@ void FDefaultInstallBundleManager::MountPaks(FContentRequestRef Request)
 	{
 		if (!bIsAsyncMount)
 		{
-			TValueOrError<void, FString> MountResult = MakeValue();
+			TValueOrError<void, UE::UnifiedError::FError> MountResult = MakeValue();
 			const bool bMountedPaks = MountPaksWorkFunc(MountResult);
 			OnPakMountCompleteFunc(bMountedPaks, MountResult);
 		}
@@ -1649,7 +1648,7 @@ void FDefaultInstallBundleManager::MountPaks(FContentRequestRef Request)
 				[MountPaksWorkFunc = MoveTemp(MountPaksWorkFunc)
 				, OnPakMountCompleteFunc = MoveTemp(OnPakMountCompleteFunc)]() mutable
 			{
-				TValueOrError<void, FString> MountResult = MakeValue();
+				TValueOrError<void, UE::UnifiedError::FError> MountResult = MakeValue();
 				const bool bMountedPaks = MountPaksWorkFunc(MountResult);
 				ExecuteOnGameThread(UE_SOURCE_LOCATION, 
 					[OnPakMountCompleteFunc = MoveTemp(OnPakMountCompleteFunc)
@@ -1776,7 +1775,7 @@ void FDefaultInstallBundleManager::MountPaks(FContentRequestRef Request)
 	}
 }
 
-bool FDefaultInstallBundleManager::MountPaksInList(TArrayView<TPair<FString, FPakMountOptions>> Paths, TValueOrError<void, FString>& OutResult, ELogVerbosity::Type LogVerbosityOverride)
+bool FDefaultInstallBundleManager::MountPaksInList(TArrayView<TPair<FString, FPakMountOptions>> Paths, TValueOrError<void, UE::UnifiedError::FError>& OutResult, ELogVerbosity::Type LogVerbosityOverride)
 {
 	if (Paths.Num() > 0 && !FCoreDelegates::MountPaksEx.IsBound())
 	{
@@ -1805,7 +1804,7 @@ bool FDefaultInstallBundleManager::MountPaksInList(TArrayView<TPair<FString, FPa
 	{
 		if (Args.Result.HasError())
 		{
-			LOG_INSTALL_BUNDLE_MAN_OVERRIDE(LogVerbosityOverride, Error, TEXT("Failed to mount pak %s, error: %s"), Args.PakFilePath, *Args.Result.GetError());
+			LOG_INSTALL_BUNDLE_MAN_OVERRIDE(LogVerbosityOverride, Error, TEXT("Failed to mount pak %s, error: %s"), Args.PakFilePath, *Args.Result.GetError().GetErrorMessage().ToString());
 
 			if (!OutResult.HasError())
 			{
diff --git a/Engine/Plugins/Experimental/DefaultInstallBundleManager/Source/Public/DefaultInstallBundleManager.h b/Engine/Plugins/Experimental/DefaultInstallBundleManager/Source/Public/DefaultInstallBundleManager.h
--- a/Engine/Plugins/Experimental/DefaultInstallBundleManager/Source/Public/DefaultInstallBundleManager.h
+++ b/Engine/Plugins/Experimental/DefaultInstallBundleManager/Source/Public/DefaultInstallBundleManager.h
@@ -6,6 +6,8 @@
 #include "InstallBundleManagerUtil.h"
 #include "InstallBundleCache.h"
 
+#include "Experimental/UnifiedError/UnifiedError.h"
+
 #include "Containers/StaticArray.h"
 #include "Templates/ValueOrError.h"
 
@@ -558,7 +560,7 @@ protected:
 
 	void MountPaks(FContentRequestRef Request);
 
-	static bool MountPaksInList(TArrayView<TPair<FString, FPakMountOptions>> Paths, TValueOrError<void, FString>& OutResult, ELogVerbosity::Type LogVerbosityOverride);
+	static bool MountPaksInList(TArrayView<TPair<FString, FPakMountOptions>> Paths, TValueOrError<void, UE::UnifiedError::FError>& OutResult, ELogVerbosity::Type LogVerbosityOverride);
 
 	void UnmountPaks(FContentReleaseRequestRef Request);
 
diff --git a/Engine/Source/Runtime/Core/Public/GenericPlatform/GenericPlatformFile.h b/Engine/Source/Runtime/Core/Public/GenericPlatform/GenericPlatformFile.h
--- a/Engine/Source/Runtime/Core/Public/GenericPlatform/GenericPlatformFile.h
+++ b/Engine/Source/Runtime/Core/Public/GenericPlatform/GenericPlatformFile.h
@@ -1001,6 +1001,7 @@ struct FPakMountOptions
 		None = 0,
 		/** Make soft references available */
 		WithSoftReferences = 1 << 0,
+		ReportDecryptionFailure = 1 << 1,
 	};
 
 	EMountFlags MountFlags = EMountFlags::None;
diff --git a/Engine/Source/Runtime/Core/Public/Misc/CoreDelegates.h b/Engine/Source/Runtime/Core/Public/Misc/CoreDelegates.h
--- a/Engine/Source/Runtime/Core/Public/Misc/CoreDelegates.h
+++ b/Engine/Source/Runtime/Core/Public/Misc/CoreDelegates.h
@@ -8,6 +8,7 @@
 #include "Containers/Map.h"
 #include "Containers/Set.h"
 #include "Containers/UnrealString.h"
+#include "Experimental/UnifiedError/UnifiedError.h"
 #include "CoreTypes.h"
 #include "Delegates/Delegate.h"
 #include "GenericPlatform/GenericPlatformFile.h"
@@ -93,7 +94,7 @@ struct FMountPaksExArgs
 	int32 Order = INDEX_NONE;
 	FPakMountOptions MountOptions;
 	// TODO: @Daniel.Lamb - this should use FError
-	TValueOrError<FMountResult, FString> Result = MakeValue();
+	TValueOrError<FMountResult, UE::UnifiedError::FError> Result = MakeValue();
 };
 }
 
diff --git a/Engine/Source/Runtime/Engine/Classes/Engine/Experimental/StreamableManagerError.h b/Engine/Source/Runtime/Engine/Classes/Engine/Experimental/StreamableManagerError.h
--- a/Engine/Source/Runtime/Engine/Classes/Engine/Experimental/StreamableManagerError.h
+++ b/Engine/Source/Runtime/Engine/Classes/Engine/Experimental/StreamableManagerError.h
@@ -19,8 +19,8 @@ DECLARE_ERROR_ONEPARAM(MissingDownloadParams, 5, StreamableManager, NSLOCTEXT("S
 DECLARE_ERROR(IoStoreNotFound, 6, StreamableManager, NSLOCTEXT("StreamableManager", "IoStoreNotFound", "IoStore did not load correctly."));
 DECLARE_ERROR_ONEPARAM(SyncLoadIncomplete, 7, StreamableManager, NSLOCTEXT("StreamableManager", "SyncLoadIncomplete", "Sync load did not complete correctly for {DebugName}."), FString, DebugName, TEXT("Unknown"));
 
-DECLARE_ERROR(AsyncLoadFailed, 8, StreamableManager, NSLOCTEXT("StreamableManager", "AsyncLoadFailed", "Async load failed."));
-DECLARE_ERROR(AsyncLoadCancelled, 9, StreamableManager, NSLOCTEXT("StreamableManager", "AsyncLoadCancelled", "Async load cancelled."));
+DECLARE_ERROR(AsyncLoadFailed, 8, StreamableManager, NSLOCTEXT("StreamableManager", "AsyncLoadFailed", "Async load failed"));
+DECLARE_ERROR(AsyncLoadCancelled, 9, StreamableManager, NSLOCTEXT("StreamableManager", "AsyncLoadCancelled", "Async load cancelled"));
 DECLARE_ERROR_ONEPARAM(AsyncLoadUnknownError, 10, StreamableManager, NSLOCTEXT("StreamableManager", "AsyncLoadUnknownError", "Unknown async loading error {AsyncLoadingErrorId}."), int32, AsyncLoadingErrorId, -1);
 
 
diff --git a/Engine/Source/Runtime/PakFile/Private/IPlatformFilePak.cpp b/Engine/Source/Runtime/PakFile/Private/IPlatformFilePak.cpp
--- a/Engine/Source/Runtime/PakFile/Private/IPlatformFilePak.cpp
+++ b/Engine/Source/Runtime/PakFile/Private/IPlatformFilePak.cpp
@@ -45,6 +45,7 @@
 #include "FilePackageStore.h"
 #include "Compression/OodleDataCompression.h"
 #include "IO/IoStore.h"
+#include "Experimental/IO/IoStatusError.h"
 #include "String/RemoveFrom.h"
 #include "Algo/AnyOf.h"
 
@@ -6485,15 +6486,17 @@ bool UE::PakFile::Private::FHandleMountPaksExDelegate::HandleDelegate(
 		{
 			Args.Result = MakeValue(UE::FMountPaksExArgs::FMountResult{ Pak.PakFile.GetReference() });
 		}
-		else if (IoMountStatus.GetErrorCode() == EIoErrorCode::PendingEncryptionKey)
+		else if ((IoMountStatus.GetErrorCode() == EIoErrorCode::PendingEncryptionKey) && 
+			!EnumHasAnyFlags(PakMountArgs.MountOptions.MountFlags, FPakMountOptions::EMountFlags::ReportDecryptionFailure))
 		{
 			// Don't consider this an error, but we have no pak to return
 			Args.Result = MakeValue(UE::FMountPaksExArgs::FMountResult{ nullptr });
 		}
 		else
 		{
-			// TODO: @Daniel.Lamb use FError here
-			Args.Result = MakeError(IoMountStatus.ToString());
+			// Experimental error handling
+			UE::UnifiedError::FError Error = UE::UnifiedError::IoStore::ConvertError(IoMountStatus);
+			Args.Result = MakeError(MoveTemp(Error));
 		}
 	}
 
