Fix zen streaming issue caused by changes to DefaultInstallBundleManager.

diff --git a/Engine/Plugins/Experimental/DefaultInstallBundleManager/Source/Private/DefaultInstallBundleManager.cpp b/Engine/Plugins/Experimental/DefaultInstallBundleManager/Source/Private/DefaultInstallBundleManager.cpp
--- a/Engine/Plugins/Experimental/DefaultInstallBundleManager/Source/Private/DefaultInstallBundleManager.cpp
+++ b/Engine/Plugins/Experimental/DefaultInstallBundleManager/Source/Private/DefaultInstallBundleManager.cpp
@@ -1777,13 +1777,18 @@ void FDefaultInstallBundleManager::MountPaks(FContentRequestRef Request)
 
 bool FDefaultInstallBundleManager::MountPaksInList(TArrayView<TPair<FString, FPakMountOptions>> Paths, TValueOrError<void, UE::UnifiedError::FError>& OutResult, ELogVerbosity::Type LogVerbosityOverride)
 {
-	if (Paths.Num() > 0 && !FCoreDelegates::MountPaksEx.IsBound())
+	OutResult = MakeValue();
+
+	if (Paths.IsEmpty())
 	{
-		ensureMsgf(false, TEXT("Pak files have not been correctly initalized. Use -UsePaks on the cmdline if you are using the UnrealEditor.exe"));
 		return false;
 	}
 
-	OutResult = MakeValue();
+	if (!FCoreDelegates::MountPaksEx.IsBound())
+	{
+		ensureMsgf(false, TEXT("Pak files have not been correctly initalized. Use -UsePaks on the cmdline if you are using the UnrealEditor.exe"));
+		return false;
+	}
 
 	TArray<UE::FMountPaksExArgs> MountArgs;
 	MountArgs.Reserve(Paths.Num());
