Fix use of moved from value

diff --git a/Engine/Plugins/Experimental/DefaultInstallBundleManager/Source/Private/DefaultInstallBundleManager.cpp b/Engine/Plugins/Experimental/DefaultInstallBundleManager/Source/Private/DefaultInstallBundleManager.cpp
--- a/Engine/Plugins/Experimental/DefaultInstallBundleManager/Source/Private/DefaultInstallBundleManager.cpp
+++ b/Engine/Plugins/Experimental/DefaultInstallBundleManager/Source/Private/DefaultInstallBundleManager.cpp
@@ -1263,7 +1263,7 @@ void FDefaultInstallBundleManager::UpdateBundleSourceComplete(TSharedRef<IInstal
 					// Only, send on success - if the request were canceled or something, there is no way to tell if the source got far enough
 					// to tell us if we had a cache hit.  We could just go by sizes in the cache, but that will
 					// be wrong in the case that we need to patch but the size stayed the same.
-					const bool bCacheHit = !InResultInfo.bContentWasInstalled;
+					const bool bCacheHit = !ResultInfo.bContentWasInstalled;
 					if (bCacheHit)
 					{
 						InstallBundleManagerAnalytics::FireEvent_BundleCacheHit(AnalyticsProvider.Get(),
