DefaultInstallBundleManager can now report if it updated content

#rnx

diff --git a/Engine/Plugins/Experimental/DefaultInstallBundleManager/Source/Private/DefaultInstallBundleManager.cpp b/Engine/Plugins/Experimental/DefaultInstallBundleManager/Source/Private/DefaultInstallBundleManager.cpp
--- a/Engine/Plugins/Experimental/DefaultInstallBundleManager/Source/Private/DefaultInstallBundleManager.cpp
+++ b/Engine/Plugins/Experimental/DefaultInstallBundleManager/Source/Private/DefaultInstallBundleManager.cpp
@@ -3463,6 +3463,12 @@ TValueOrError<FInstallBundleRequestInfo, EInstallBundleResult> FDefaultInstallBu
 				RequestSharedContext = MakeShared<InstallBundleUtil::FContentRequestSharedContext>();
 			}
 			NewRequest->RequestSharedContext = RequestSharedContext;
+
+			const EBundleState BundleState = GetBundleStatus(*BundleInfo);
+			if (BundleState == EBundleState::NeedsUpdate || BundleState == EBundleState::NotInstalled)
+			{
+				bHasEverUpdatedContent = true;
+			}
 		}
 
 		RetInfo.InfoFlags |= EInstallBundleRequestInfoFlags::EnqueuedBundles;
diff --git a/Engine/Plugins/Experimental/DefaultInstallBundleManager/Source/Public/DefaultInstallBundleManager.h b/Engine/Plugins/Experimental/DefaultInstallBundleManager/Source/Public/DefaultInstallBundleManager.h
--- a/Engine/Plugins/Experimental/DefaultInstallBundleManager/Source/Public/DefaultInstallBundleManager.h
+++ b/Engine/Plugins/Experimental/DefaultInstallBundleManager/Source/Public/DefaultInstallBundleManager.h
@@ -702,6 +702,9 @@ public:
 #if !UE_BUILD_SHIPPING
 	UE_API virtual void GetDebugText(TArray<FString>& Output) override;
 #endif
+
+	virtual bool HasEverUpdatedContent() const { return bHasEverUpdatedContent; }
+
 protected:
 	//Special version of these to wrap our calls to PersistentStats
 	UE_API void StartBundlePersistentStatTracking(TSharedRef<FContentRequest> ContentRequest, const FString& ExpectedAnalyticsID = FString(), bool bForceResetStatData = false);
@@ -800,6 +803,8 @@ protected:
 	//this functionality even on ship builds
 	bool bOverrideCommand_SkipPatchCheck = false;
 
+	bool bHasEverUpdatedContent = false;
+
 protected:
 	// Analytics
 	TSharedPtr<IAnalyticsProviderET> AnalyticsProvider;
diff --git a/Engine/Source/Runtime/InstallBundleManager/Public/InstallBundleManagerInterface.h b/Engine/Source/Runtime/InstallBundleManager/Public/InstallBundleManagerInterface.h
--- a/Engine/Source/Runtime/InstallBundleManager/Public/InstallBundleManagerInterface.h
+++ b/Engine/Source/Runtime/InstallBundleManager/Public/InstallBundleManagerInterface.h
@@ -228,5 +228,7 @@ PRAGMA_ENABLE_DEPRECATION_WARNINGS
 #if !UE_BUILD_SHIPPING
 	virtual void GetDebugText(TArray<FString>& Output) {}
 #endif
+
+	virtual bool HasEverUpdatedContent() const { return true; }
 };
 
