Adding a new function to InstallBundleManager to return the list of bundle names are currently pending removal on next init

diff --git a/Engine/Plugins/Experimental/DefaultInstallBundleManager/Source/Private/DefaultInstallBundleManager.cpp b/Engine/Plugins/Experimental/DefaultInstallBundleManager/Source/Private/DefaultInstallBundleManager.cpp
--- a/Engine/Plugins/Experimental/DefaultInstallBundleManager/Source/Private/DefaultInstallBundleManager.cpp
+++ b/Engine/Plugins/Experimental/DefaultInstallBundleManager/Source/Private/DefaultInstallBundleManager.cpp
@@ -3961,6 +3961,19 @@ void FDefaultInstallBundleManager::CancelRequestRemoveContentOnNextInit(TArrayVi
 	GConfig->Flush(false);
 }
 
+TArray<FName> FDefaultInstallBundleManager::GetRequestedRemoveContentOnNextInit() const
+{
+	TArray<FString> ConfigBundlesToRemove;
+	GConfig->GetArray(TEXT("InstallBundleManager.UserSettings"), TEXT("RemoveBundleOnInit"), ConfigBundlesToRemove, GGameUserSettingsIni);
+
+	TArray<FName> ConfigBundleNamesToRemove;
+	Algo::Transform(ConfigBundlesToRemove, ConfigBundleNamesToRemove, [](const FString& Item) {
+		return FName(Item);
+	});
+
+	return ConfigBundleNamesToRemove;
+}
+
 void FDefaultInstallBundleManager::CancelUpdateContent(TArrayView<const FName> BundleNames)
 {
 	CancelUpdateContentInternal(BundleNames);
diff --git a/Engine/Plugins/Experimental/DefaultInstallBundleManager/Source/Public/DefaultInstallBundleManager.h b/Engine/Plugins/Experimental/DefaultInstallBundleManager/Source/Public/DefaultInstallBundleManager.h
--- a/Engine/Plugins/Experimental/DefaultInstallBundleManager/Source/Public/DefaultInstallBundleManager.h
+++ b/Engine/Plugins/Experimental/DefaultInstallBundleManager/Source/Public/DefaultInstallBundleManager.h
@@ -666,6 +666,8 @@ public:
 
 	UE_API virtual void CancelRequestRemoveContentOnNextInit(TArrayView<const FName> BundleName) override;
 
+	UE_API virtual TArray<FName> GetRequestedRemoveContentOnNextInit() const override;
+
 	UE_API virtual void CancelUpdateContent(TArrayView<const FName> BundleNames) override;
 
 	UE_API virtual void PauseUpdateContent(TArrayView<const FName> BundleNames) override;
diff --git a/Engine/Source/Runtime/InstallBundleManager/Public/InstallBundleManagerInterface.h b/Engine/Source/Runtime/InstallBundleManager/Public/InstallBundleManagerInterface.h
--- a/Engine/Source/Runtime/InstallBundleManager/Public/InstallBundleManagerInterface.h
+++ b/Engine/Source/Runtime/InstallBundleManager/Public/InstallBundleManagerInterface.h
@@ -183,6 +183,8 @@ PRAGMA_ENABLE_DEPRECATION_WARNINGS
 	INSTALLBUNDLEMANAGER_API void CancelRequestRemoveContentOnNextInit(FName BundleName);
 	virtual void CancelRequestRemoveContentOnNextInit(TArrayView<const FName> BundleNames) = 0;
 
+	virtual TArray<FName> GetRequestedRemoveContentOnNextInit() const = 0;
+
 	INSTALLBUNDLEMANAGER_API void CancelUpdateContent(FName BundleName);
 	virtual void CancelUpdateContent(TArrayView<const FName> BundleNames) = 0;
 
diff --git a/Engine/Source/Runtime/InstallBundleManager/Public/InstallBundleSourceInterface.h b/Engine/Source/Runtime/InstallBundleManager/Public/InstallBundleSourceInterface.h
--- a/Engine/Source/Runtime/InstallBundleManager/Public/InstallBundleSourceInterface.h
+++ b/Engine/Source/Runtime/InstallBundleManager/Public/InstallBundleSourceInterface.h
@@ -115,6 +115,8 @@ public:
 	// Call to cancel the removal of any content scheduled for removal the next time the source is initialized
 	// Returns true if all bundles were canceled
 	virtual bool CancelRequestRemoveContentOnNextInit(TArrayView<const FName> BundleNames) { return false; }
+	// Returns the bundle names of any content that is currently removal the next time the source is initialized
+	virtual TArray<FName> GetRequestedRemoveContentOnNextInit() const { return {}; }
 
 	// Cancel the install for the specified bundles
 	virtual void CancelBundles(TArrayView<const FName> BundleNames) {}
diff --git a/Engine/Source/Runtime/NullInstallBundleManager/Private/NullInstallBundleManager.cpp b/Engine/Source/Runtime/NullInstallBundleManager/Private/NullInstallBundleManager.cpp
--- a/Engine/Source/Runtime/NullInstallBundleManager/Private/NullInstallBundleManager.cpp
+++ b/Engine/Source/Runtime/NullInstallBundleManager/Private/NullInstallBundleManager.cpp
@@ -114,6 +114,11 @@ PRAGMA_ENABLE_DEPRECATION_WARNINGS
 	{
 	}
 
+	virtual TArray<FName> GetRequestedRemoveContentOnNextInit() const override
+	{
+		return TArray<FName>();
+	}
+
 	virtual void CancelUpdateContent(TArrayView<const FName> BundleNames) override
 	{
 	}
