Monitor CVar changes to trigger re-installing the cosmetic bundle in order to mount file or on-demand container TOC.

#rnx

diff --git a/Engine/Plugins/Experimental/DefaultInstallBundleManager/Source/Private/DefaultInstallBundleManager.cpp b/Engine/Plugins/Experimental/DefaultInstallBundleManager/Source/Private/DefaultInstallBundleManager.cpp
--- a/Engine/Plugins/Experimental/DefaultInstallBundleManager/Source/Private/DefaultInstallBundleManager.cpp
+++ b/Engine/Plugins/Experimental/DefaultInstallBundleManager/Source/Private/DefaultInstallBundleManager.cpp
@@ -1843,6 +1843,20 @@ void FDefaultInstallBundleManager::UnmountPaks(FContentReleaseRequestRef Request
 		}
 	}
 
+	if (BundleInfo.bMountedOnDemandTocs)
+	{
+		BundleInfo.bMountedOnDemandTocs = false;
+		if (UE::IoStore::IOnDemandIoStore* IoStore = UE::IoStore::TryGetOnDemandIoStore())
+		{
+			if (FIoStatus Status = IoStore->Unmount(Request->BundleName.ToString()); !Status.IsOk())
+			{
+				LOG_INSTALL_BUNDLE_MAN_OVERRIDE(Request->LogVerbosityOverride, Error,
+					TEXT("Error unmounting ondemand toc for bundle '%s' : %s"),
+					*BundleInfo.BundleNameString, *Status.ToString());
+			}
+		}
+	}
+
 	OnPaksUnmountedInternal(Request, BundleInfo);
 
 	SetBundleStatus(BundleInfo, EBundleState::NeedsMount);
