Add OnPaksMountingInternal to allow custom install bundle manager behaviour right before paks are mounted for an install bundle

diff --git a/Engine/Plugins/Experimental/DefaultInstallBundleManager/Source/Private/DefaultInstallBundleManager.cpp b/Engine/Plugins/Experimental/DefaultInstallBundleManager/Source/Private/DefaultInstallBundleManager.cpp
--- a/Engine/Plugins/Experimental/DefaultInstallBundleManager/Source/Private/DefaultInstallBundleManager.cpp
+++ b/Engine/Plugins/Experimental/DefaultInstallBundleManager/Source/Private/DefaultInstallBundleManager.cpp
@@ -1631,15 +1631,28 @@ void FDefaultInstallBundleManager::MountPaks(FContentRequestRef Request)
 		Request->StepResult = EContentRequestStepResult::Done;
 	};
 
+	auto OnPakMountingFunc = [this, WeakThis = AsWeak(), Request]() mutable
+	{
+		TSharedPtr<IInstallBundleManager> PinnedThis = WeakThis.Pin();
+		if (!PinnedThis)
+		{
+			return;
+		}
+		FBundleInfo& BundleInfoLambda = BundleInfoMap.FindChecked(Request->BundleName);
+		OnPaksMountingInternal(Request, BundleInfoLambda);
+	};
+
 	const bool bIsAsyncMount = EnumHasAnyFlags(Request->Flags, EInstallBundleRequestFlags::AsyncMount);
 	TUniqueFunction<void()> StartMountPaksFunc = 
 		[MountPaksWorkFunc = MoveTemp(MountPaksWorkFunc)
 		, OnPakMountCompleteFunc = MoveTemp(OnPakMountCompleteFunc)
+		, OnPakMountingFunc = MoveTemp(OnPakMountingFunc)
 		, bIsAsyncMount]() mutable
 	{
 		if (!bIsAsyncMount)
 		{
 			TValueOrError<void, UE::UnifiedError::FError> MountResult = MakeValue();
+			OnPakMountingFunc();
 			const bool bMountedPaks = MountPaksWorkFunc(MountResult);
 			OnPakMountCompleteFunc(bMountedPaks, MountResult);
 		}
@@ -1647,9 +1660,11 @@ void FDefaultInstallBundleManager::MountPaks(FContentRequestRef Request)
 		{
 			InstallBundleUtil::StartInstallBundleAsyncIOTask(
 				[MountPaksWorkFunc = MoveTemp(MountPaksWorkFunc)
-				, OnPakMountCompleteFunc = MoveTemp(OnPakMountCompleteFunc)]() mutable
+				, OnPakMountCompleteFunc = MoveTemp(OnPakMountCompleteFunc)
+				, OnPakMountingFunc = MoveTemp(OnPakMountingFunc)]() mutable
 			{
 				TValueOrError<void, UE::UnifiedError::FError> MountResult = MakeValue();
+				OnPakMountingFunc();
 				const bool bMountedPaks = MountPaksWorkFunc(MountResult);
 				ExecuteOnGameThread(UE_SOURCE_LOCATION, 
 					[OnPakMountCompleteFunc = MoveTemp(OnPakMountCompleteFunc)
diff --git a/Engine/Plugins/Experimental/DefaultInstallBundleManager/Source/Public/DefaultInstallBundleManager.h b/Engine/Plugins/Experimental/DefaultInstallBundleManager/Source/Public/DefaultInstallBundleManager.h
--- a/Engine/Plugins/Experimental/DefaultInstallBundleManager/Source/Public/DefaultInstallBundleManager.h
+++ b/Engine/Plugins/Experimental/DefaultInstallBundleManager/Source/Public/DefaultInstallBundleManager.h
@@ -570,6 +570,7 @@ protected:
 
 	UE_API virtual TArray<TPair<FString, FPakMountOptions>> GetPakMountList(FContentRequestRef Request, const FBundleInfo& BundleInfo);
 
+	virtual void OnPaksMountingInternal(FContentRequestRef Request, FBundleInfo& BundleInfo) {}
 	virtual void OnPaksMountedInternal(FContentRequestRef Request, FBundleInfo& BundleInfo) {}
 	virtual void OnPaksUnmountedInternal(FContentReleaseRequestRef Request, FBundleInfo& BundleInfo) {};
 
