InstallBundleManager - Add unique request cookie to create unique stat keys for every install bundle request.  This should fix problems where multiple requests are flushing through the pipeline and may stomp each others keys.  Usually this can't happen beacuse only one active request per bundle is allowed, but there are rare cases where multiple cancelled requests can stack up and may leap-frog each other as they flush through. These corner cases may lead to a crash.

diff --git a/Engine/Plugins/Experimental/DefaultInstallBundleManager/Source/Private/DefaultInstallBundleManager.cpp b/Engine/Plugins/Experimental/DefaultInstallBundleManager/Source/Private/DefaultInstallBundleManager.cpp
--- a/Engine/Plugins/Experimental/DefaultInstallBundleManager/Source/Private/DefaultInstallBundleManager.cpp
+++ b/Engine/Plugins/Experimental/DefaultInstallBundleManager/Source/Private/DefaultInstallBundleManager.cpp
@@ -909,7 +909,7 @@ void FDefaultInstallBundleManager::AddRequestToInitialBatch(FContentRequestRef R
 
 void FDefaultInstallBundleManager::ReserveCache(FContentRequestRef Request)
 {
-	StatsBegin(Request->BundleName, EContentRequestState::ReservingCache);
+	StatsBegin(Request->GetStatsKey(), EContentRequestState::ReservingCache);
 
 	TryReserveCache(Request);
 }
@@ -1001,7 +1001,7 @@ void FDefaultInstallBundleManager::TryReserveCache(FContentRequestRef Request)
 
 		Request->LastCacheReserveResult = EInstallBundleCacheReserveResult::Fail_CacheFull;
 
-		StatsEnd(Request->BundleName, EContentRequestState::ReservingCache);
+		StatsEnd(Request->GetStatsKey(), EContentRequestState::ReservingCache);
 		Request->Result = EInstallBundleResult::FailedCacheReserve;
 		Request->StepResult = EContentRequestStepResult::Done;
 		return;
@@ -1016,7 +1016,7 @@ void FDefaultInstallBundleManager::TryReserveCache(FContentRequestRef Request)
 
 	if (Request->BundlesToEvictFromSourcesMap.Num() == 0)
 	{
-		StatsEnd(Request->BundleName, EContentRequestState::ReservingCache);
+		StatsEnd(Request->GetStatsKey(), EContentRequestState::ReservingCache);
 		Request->StepResult = EContentRequestStepResult::Done;
 		return;
 	}
@@ -1162,7 +1162,7 @@ void FDefaultInstallBundleManager::CacheEvictionComplete(TSharedRef<IInstallBund
 		FContentRequestRef ContentRequest = StaticCastSharedRef<FContentRequest>(Requestor);
 		if (ContentRequest->bIsCanceled) // If canceled don't keep trying to reserve, just finish so we can release any reserves that succeeded ASAP
 		{
-			StatsEnd(ContentRequest->BundleName, EContentRequestState::ReservingCache);
+			StatsEnd(ContentRequest->GetStatsKey(), EContentRequestState::ReservingCache);
 			ContentRequest->StepResult = EContentRequestStepResult::Done;
 		}
 		else
@@ -1198,7 +1198,7 @@ TArray<TSharedPtr<IInstallBundleSource>> FDefaultInstallBundleManager::GetEnable
 
 void FDefaultInstallBundleManager::UpdateBundleSources(FContentRequestRef Request)
 {
-	StatsBegin(Request->BundleName, EContentRequestState::UpdatingBundleSources);
+	StatsBegin(Request->GetStatsKey(), EContentRequestState::UpdatingBundleSources);
 
 	Request->StepResult = EContentRequestStepResult::Waiting;
 	TArray<TSharedPtr<IInstallBundleSource>> EnabledBundleSources = GetEnabledBundleSourcesForRequest(Request);
@@ -1214,6 +1214,7 @@ void FDefaultInstallBundleManager::UpdateBundleSources(FContentRequestRef Reques
 		Context.PausedCallback.BindRaw(this, &FDefaultInstallBundleManager::UpdateBundleSourcePause, Request);
 		Context.CompleteCallback.BindRaw(this, &FDefaultInstallBundleManager::UpdateBundleSourceComplete, Request);
 		Context.RequestSharedContext = Request->RequestSharedContext;
+		Context.StatsCookie = Request->GetStatsCookie();
 
 		Source->RequestUpdateContent(MoveTemp(Context));
 	}
@@ -1362,7 +1363,7 @@ void FDefaultInstallBundleManager::UpdateBundleSourceComplete(TSharedRef<IInstal
 
 	Request->StepResult = EContentRequestStepResult::Done;
 
-	StatsEnd(Request->BundleName, EContentRequestState::UpdatingBundleSources);
+	StatsEnd(Request->GetStatsKey(), EContentRequestState::UpdatingBundleSources);
 
 	PersistentStats->UpdateForBundleSource(InResultInfo, SourceType, Request->BundleName.ToString());
 }
@@ -1611,7 +1612,7 @@ void FDefaultInstallBundleManager::MountPaks(FContentRequestRef Request)
 			Request->OptionalErrorText = MountResult.GetError().GetErrorMessage(); 
 			Request->StepResult = EContentRequestStepResult::Done;
 
-			StatsEnd(Request->BundleName, EContentRequestState::Mounting);
+			StatsEnd(Request->GetStatsKey(), EContentRequestState::Mounting);
 			TRACE_BOOKMARK(TEXT("Finished Mount Bundle %s"), *Request->BundleName.ToString());
 			return;
 		}
@@ -1626,7 +1627,7 @@ void FDefaultInstallBundleManager::MountPaks(FContentRequestRef Request)
 		// Update bundle status
 		SetBundleStatus(BundleInfoLambda, EBundleState::Mounted);
 
-		StatsEnd(Request->BundleName, EContentRequestState::Mounting);
+		StatsEnd(Request->GetStatsKey(), EContentRequestState::Mounting);
 		TRACE_BOOKMARK(TEXT("Finished Mount Bundle %s"), *Request->BundleName.ToString());
 		Request->StepResult = EContentRequestStepResult::Done;
 	};
@@ -1685,7 +1686,7 @@ void FDefaultInstallBundleManager::MountPaks(FContentRequestRef Request)
 	if (Request->OnDemandMountArgs.IsEmpty() || BundleInfo.bMountedOnDemandTocs || !AllowIoStoreOnDemandMount(Request, BundleInfo))
 	{
 		LOG_INSTALL_BUNDLE_MAN_OVERRIDE(Request->LogVerbosityOverride, Display, TEXT("Mounting Paks for Request %s"), *BundleInfo.BundleNameString);
-		StatsBegin(Request->BundleName, EContentRequestState::Mounting);
+		StatsBegin(Request->GetStatsKey(), EContentRequestState::Mounting);
 		StartMountPaksFunc();
 		return;
 	}
@@ -1703,7 +1704,7 @@ void FDefaultInstallBundleManager::MountPaks(FContentRequestRef Request)
 	}
 
 	LOG_INSTALL_BUNDLE_MAN_OVERRIDE(Request->LogVerbosityOverride, Display, TEXT("Mounting OnDemand Tocs for Request %s"), *BundleInfo.BundleNameString);
-	StatsBegin(Request->BundleName, EContentRequestState::Mounting);
+	StatsBegin(Request->GetStatsKey(), EContentRequestState::Mounting);
 
 	struct FOnDemandMountContext
 	{
@@ -1737,7 +1738,7 @@ void FDefaultInstallBundleManager::MountPaks(FContentRequestRef Request)
 						FString ErrorCode = GetIoErrorText(MountResult.Status.GetErrorCode());
 						ErrorCode.ReplaceCharInline(TEXT(' '), TEXT('_'), ESearchCase::CaseSensitive);
 
-						PinnedThis->StatsEnd(Request->BundleName, EContentRequestState::Mounting);
+						PinnedThis->StatsEnd(Request->GetStatsKey(), EContentRequestState::Mounting);
 						TRACE_BOOKMARK(TEXT("Finished Mount Bundle %s"), *Request->BundleName.ToString());
 
 						Request->Result = EInstallBundleResult::InstallError;
@@ -1755,7 +1756,7 @@ void FDefaultInstallBundleManager::MountPaks(FContentRequestRef Request)
 				{
 					LOG_INSTALL_BUNDLE_MAN_OVERRIDE(Request->LogVerbosityOverride, Display, TEXT("Skipping Mounting Paks for Request %s"), *Request->BundleName.ToString());
 
-					PinnedThis->StatsEnd(Request->BundleName, EContentRequestState::Mounting);
+					PinnedThis->StatsEnd(Request->GetStatsKey(), EContentRequestState::Mounting);
 					TRACE_BOOKMARK(TEXT("Finished Mount Bundle %s"), *Request->BundleName.ToString());
 					Request->StepResult = EContentRequestStepResult::Done;
 					return;
@@ -1898,7 +1899,7 @@ void FDefaultInstallBundleManager::WaitForShaderCache(FContentRequestRef Request
 		return;
 	}
 
-	StatsBegin(Request->BundleName, EContentRequestState::WaitingForShaderCache);
+	StatsBegin(Request->GetStatsKey(), EContentRequestState::WaitingForShaderCache);
 	PersistentTimingStatsBegin(Request, InstallBundleUtil::PersistentStats::ETimingStatNames::PSOTime_Real);
 
 	int32 NumPrecompilesRemaining = FShaderPipelineCache::NumPrecompilesRemaining();
@@ -1920,7 +1921,7 @@ void FDefaultInstallBundleManager::WaitForShaderCache(FContentRequestRef Request
 
 		SetMustWaitForPSOCache(BundleInfo, 0); // Make sure we pass 0 to flag that there are no more shaders to wait on
 
-		StatsEnd(Request->BundleName, EContentRequestState::WaitingForShaderCache);
+		StatsEnd(Request->GetStatsKey(), EContentRequestState::WaitingForShaderCache);
 		PersistentTimingStatsEnd(Request, InstallBundleUtil::PersistentStats::ETimingStatNames::PSOTime_Real);
 		Request->StepResult = EContentRequestStepResult::Done;
 	}
@@ -1928,7 +1929,7 @@ void FDefaultInstallBundleManager::WaitForShaderCache(FContentRequestRef Request
 
 void FDefaultInstallBundleManager::FinishRequest(FContentRequestRef Request)
 {
-	StatsBegin(Request->BundleName, EContentRequestState::Finishing);
+	StatsBegin(Request->GetStatsKey(), EContentRequestState::Finishing);
 
 	if (Request->Result == EInstallBundleResult::OK || Request->Result == EInstallBundleResult::UserCancelledError)
 	{
@@ -1971,7 +1972,7 @@ void FDefaultInstallBundleManager::FinishRequest(FContentRequestRef Request)
 		InstallBundleCompleteDelegate.Broadcast(MoveTemp(ResultInfo));
 	}
 
-	StatsEnd(Request->BundleName, EContentRequestState::Finishing);
+	StatsEnd(Request->GetStatsKey(), EContentRequestState::Finishing);
 	PersistentTimingStatsEnd(Request, InstallBundleUtil::PersistentStats::ETimingStatNames::TotalTime_Real);
 	StopBundlePersistentStatTracking(Request);
 
@@ -2100,7 +2101,7 @@ void FDefaultInstallBundleManager::TickContentRequests()
 		{
 			if (Request->iPrereq == INDEX_NONE)
 			{
-				StatsBegin(Request->BundleName);
+				StatsBegin(Request->GetStatsKey());
 				LOG_INSTALL_BUNDLE_MAN_OVERRIDE(Request->LogVerbosityOverride, Display, TEXT("Starting Request %s"), *Request->BundleName.ToString());
 			}
 
@@ -2250,8 +2251,8 @@ void FDefaultInstallBundleManager::TickContentRequests()
 				FinishRequest(Request);
 				break;
 			case EContentRequestState::CleaningUp:
-				StatsEnd(Request->BundleName);
-				LogStats(Request->BundleName, Request->LogVerbosityOverride);
+				StatsEnd(Request->GetStatsKey());
+				LogStats(Request->GetStatsKey(), Request->LogVerbosityOverride);
 
 				CacheHintRequested(Request, false);
 
@@ -2261,11 +2262,11 @@ void FDefaultInstallBundleManager::TickContentRequests()
 						Request->BundleName.ToString(),
 						Request->bContentWasInstalled,
 						LexToString(Request->Result),
-						StatsMap->GetMap().FindChecked(Request->BundleName));
+						StatsMap->FindStatsChecked(Request->GetStatsKey()));
 				}
 
 				LOG_INSTALL_BUNDLE_MAN_OVERRIDE(Request->LogVerbosityOverride, Display, TEXT("Removing Request %s"), *Request->BundleName.ToString());
-				StatsMap->StatsReset(Request->BundleName);
+				StatsMap->StatsReset(Request->GetStatsKey());
 				ContentRequests[EContentRequestBatch::Install].RemoveAtSwap(i); //Invalidates Request
 				bRequestComplete = true;
 				break;
@@ -2800,7 +2801,7 @@ void FDefaultInstallBundleManager::TickWaitForShaderCache()
 			if (Request->bIsCanceled)
 			{
 				LOG_INSTALL_BUNDLE_MAN_OVERRIDE(Request->LogVerbosityOverride, Display, TEXT("Canceled Waiting for PSO cache for %s"), *BundleInfo.BundleNameString);
-				StatsEnd(Request->BundleName, EContentRequestState::WaitingForShaderCache);
+				StatsEnd(Request->GetStatsKey(), EContentRequestState::WaitingForShaderCache);
 				PersistentTimingStatsEnd(Request, InstallBundleUtil::PersistentStats::ETimingStatNames::PSOTime_Real);
 				Request->StepResult = EContentRequestStepResult::Done;
 			}
@@ -2817,7 +2818,7 @@ void FDefaultInstallBundleManager::TickWaitForShaderCache()
 
 				SetMustWaitForPSOCache(BundleInfo, 0); // Make sure we pass 0 to flag that there are no more shaders to wait on
 
-				StatsEnd(Request->BundleName, EContentRequestState::WaitingForShaderCache);
+				StatsEnd(Request->GetStatsKey(), EContentRequestState::WaitingForShaderCache);
 				PersistentTimingStatsEnd(Request, InstallBundleUtil::PersistentStats::ETimingStatNames::PSOTime_Real);
 				Request->StepResult = EContentRequestStepResult::Done;
 			}
@@ -5142,32 +5143,32 @@ void FDefaultInstallBundleManager::AsyncInit_FireInitAnlaytic(bool bCanRetry)
 	InitStepResult = EAsyncInitStepResult::Done;
 }
 
-void FDefaultInstallBundleManager::StatsBegin(FName BundleName)
+void FDefaultInstallBundleManager::StatsBegin(const InstallBundleUtil::FContentRequestStatsKey& Key)
 {
-	StatsMap->StatsBegin(BundleName);
+	StatsMap->StatsBegin(Key);
 }
 
-void FDefaultInstallBundleManager::StatsEnd(FName BundleName)
+void FDefaultInstallBundleManager::StatsEnd(const InstallBundleUtil::FContentRequestStatsKey& Key)
 {
-	StatsMap->StatsEnd(BundleName);
+	StatsMap->StatsEnd(Key);
 }
 
-void FDefaultInstallBundleManager::StatsBegin(FName BundleName, EContentRequestState State)
+void FDefaultInstallBundleManager::StatsBegin(const InstallBundleUtil::FContentRequestStatsKey& Key, EContentRequestState State)
 {
-	StatsMap->StatsBegin(BundleName, LexToString(State));
+	StatsMap->StatsBegin(Key, LexToString(State));
 }
 
-void FDefaultInstallBundleManager::StatsEnd(FName BundleName, EContentRequestState State, uint64 DataSize /*= 0*/)
+void FDefaultInstallBundleManager::StatsEnd(const InstallBundleUtil::FContentRequestStatsKey& Key, EContentRequestState State, uint64 DataSize /*= 0*/)
 {
-	StatsMap->StatsEnd(BundleName, LexToString(State), DataSize);
+	StatsMap->StatsEnd(Key, LexToString(State), DataSize);
 }
 
-void FDefaultInstallBundleManager::LogStats(FName BundleName, ELogVerbosity::Type LogVerbosityOverride)
+void FDefaultInstallBundleManager::LogStats(const InstallBundleUtil::FContentRequestStatsKey& Key, ELogVerbosity::Type LogVerbosityOverride)
 {
-	const InstallBundleUtil::FContentRequestStats& RequestStats = StatsMap->GetMap().FindChecked(BundleName);
+	const InstallBundleUtil::FContentRequestStats& RequestStats = StatsMap->FindStatsChecked(Key);
 	check(!RequestStats.bOpen);
 
-	InstallBundleManagerUtil::LogBundleRequestStats(*BundleName.ToString(), RequestStats, LogVerbosityOverride);
+	InstallBundleManagerUtil::LogBundleRequestStats(Key, RequestStats, LogVerbosityOverride);
 }
 
 
diff --git a/Engine/Plugins/Experimental/DefaultInstallBundleManager/Source/Private/InstallBundleManagerUtil.cpp b/Engine/Plugins/Experimental/DefaultInstallBundleManager/Source/Private/InstallBundleManagerUtil.cpp
--- a/Engine/Plugins/Experimental/DefaultInstallBundleManager/Source/Private/InstallBundleManagerUtil.cpp
+++ b/Engine/Plugins/Experimental/DefaultInstallBundleManager/Source/Private/InstallBundleManagerUtil.cpp
@@ -186,10 +186,13 @@ namespace InstallBundleManagerUtil
 		}
 	}
 
-	void LogBundleRequestStats(const TCHAR* BundleName, const InstallBundleUtil::FContentRequestStats& RequestStats, ELogVerbosity::Type LogVerbosityOverride)
+	void LogBundleRequestStats(const InstallBundleUtil::FContentRequestStatsKey& Key, const InstallBundleUtil::FContentRequestStats& RequestStats, ELogVerbosity::Type LogVerbosityOverride)
 	{
+		FNameBuilder BundleNameBuilder(Key.GetBundleName());
+		const TCHAR* BundleName = *BundleNameBuilder;
+
 		LOG_INSTALL_BUNDLE_MAN_OVERRIDE(LogVerbosityOverride, Display, TEXT("------------------------------------------------------"));
-		LOG_INSTALL_BUNDLE_MAN_OVERRIDE(LogVerbosityOverride, Display, TEXT("Bundle Request Stats - %s"), BundleName);
+		LOG_INSTALL_BUNDLE_MAN_OVERRIDE(LogVerbosityOverride, Display, TEXT("Bundle Request Stats - %s (%p)"), BundleName, Key.GetCookie());
 		LOG_INSTALL_BUNDLE_MAN_OVERRIDE(LogVerbosityOverride, Display, TEXT("Total Time %f"), RequestStats.GetElapsedTime());
 
 		if (ensureAlways(!RequestStats.bOpen))
@@ -212,7 +215,7 @@ namespace InstallBundleManagerUtil
 			LOG_INSTALL_BUNDLE_MAN_OVERRIDE(LogVerbosityOverride, Display, TEXT("Bundle Request %s did not finish"), BundleName);
 		}
 
-		LOG_INSTALL_BUNDLE_MAN_OVERRIDE(LogVerbosityOverride, Display, TEXT("End Bundle Request Stats - %s"), BundleName);
+		LOG_INSTALL_BUNDLE_MAN_OVERRIDE(LogVerbosityOverride, Display, TEXT("End Bundle Request Stats - %s (%p)"), BundleName, Key.GetCookie());
 		LOG_INSTALL_BUNDLE_MAN_OVERRIDE(LogVerbosityOverride, Display, TEXT("------------------------------------------------------"));
 	}
 
diff --git a/Engine/Plugins/Experimental/DefaultInstallBundleManager/Source/Public/DefaultInstallBundleManager.h b/Engine/Plugins/Experimental/DefaultInstallBundleManager/Source/Public/DefaultInstallBundleManager.h
--- a/Engine/Plugins/Experimental/DefaultInstallBundleManager/Source/Public/DefaultInstallBundleManager.h
+++ b/Engine/Plugins/Experimental/DefaultInstallBundleManager/Source/Public/DefaultInstallBundleManager.h
@@ -428,6 +428,9 @@ protected:
 		virtual FString GetEvictionRequestorName() const override { return BundleName.ToString(); }
 		virtual ECacheEvictionRequestorType GetEvictionRequestorType() const override { return ECacheEvictionRequestorType::ContentRequest; }
 		virtual ELogVerbosity::Type GetLogVerbosityOverride() const override { return LogVerbosityOverride;  }
+
+		void* GetStatsCookie() const { return const_cast<FContentRequest*>(this); }
+		InstallBundleUtil::FContentRequestStatsKey GetStatsKey() const { return InstallBundleUtil::FContentRequestStatsKey(BundleName, GetStatsCookie()); }
 	};
 	using FContentRequestRef = TSharedRef<FContentRequest>;
 	using FContentRequestPtr = TSharedPtr<FContentRequest>;
@@ -733,11 +736,12 @@ protected:
 	UE_API void AsyncInit_SetUpdateBundleInfoCallback();
 	UE_API void AsyncInit_CreateAnalyticsSession();
 	UE_API void AsyncInit_FireInitAnlaytic(bool bCanRetry);
-	UE_API void StatsBegin(FName BundleName);
-	UE_API void StatsEnd(FName BundleName);
-	UE_API void StatsBegin(FName BundleName, EContentRequestState State);
-	UE_API void StatsEnd(FName BundleName, EContentRequestState State, uint64 DataSize = 0);
-	UE_API void LogStats(FName BundleName, ELogVerbosity::Type LogVerbosityOverride);
+
+	UE_API void StatsBegin(const InstallBundleUtil::FContentRequestStatsKey& Key);
+	UE_API void StatsEnd(const InstallBundleUtil::FContentRequestStatsKey& Key);
+	UE_API void StatsBegin(const InstallBundleUtil::FContentRequestStatsKey& Key, EContentRequestState State);
+	UE_API void StatsEnd(const InstallBundleUtil::FContentRequestStatsKey& Key, EContentRequestState State, uint64 DataSize = 0);
+	UE_API void LogStats(const InstallBundleUtil::FContentRequestStatsKey& Key, ELogVerbosity::Type LogVerbosityOverride);
 
 protected:
 	FTSTicker::FDelegateHandle TickHandle;
diff --git a/Engine/Plugins/Experimental/DefaultInstallBundleManager/Source/Public/InstallBundleManagerUtil.h b/Engine/Plugins/Experimental/DefaultInstallBundleManager/Source/Public/InstallBundleManagerUtil.h
--- a/Engine/Plugins/Experimental/DefaultInstallBundleManager/Source/Public/InstallBundleManagerUtil.h
+++ b/Engine/Plugins/Experimental/DefaultInstallBundleManager/Source/Public/InstallBundleManagerUtil.h
@@ -64,7 +64,7 @@ namespace InstallBundleManagerUtil
 	};
 #undef JSON_MCI_VALUE
 
-	DEFAULTINSTALLBUNDLEMANAGER_API void LogBundleRequestStats(const TCHAR* BundleName, const InstallBundleUtil::FContentRequestStats& RequestStats, ELogVerbosity::Type LogVerbosityOverride);
+	DEFAULTINSTALLBUNDLEMANAGER_API void LogBundleRequestStats(const InstallBundleUtil::FContentRequestStatsKey& Key, const InstallBundleUtil::FContentRequestStats& RequestStats, ELogVerbosity::Type LogVerbosityOverride);
 
 	class FPersistentStatContainer : public InstallBundleUtil::PersistentStats::FPersistentStatContainerBase
 	{
diff --git a/Engine/Plugins/Experimental/DefaultInstallBundleManager/Source/Public/InstallBundleSourcePlatformBase.h b/Engine/Plugins/Experimental/DefaultInstallBundleManager/Source/Public/InstallBundleSourcePlatformBase.h
--- a/Engine/Plugins/Experimental/DefaultInstallBundleManager/Source/Public/InstallBundleSourcePlatformBase.h
+++ b/Engine/Plugins/Experimental/DefaultInstallBundleManager/Source/Public/InstallBundleSourcePlatformBase.h
@@ -42,15 +42,15 @@ protected:
 	virtual bool QueryPersistentBundleInfo(FInstallBundleSourcePersistentBundleInfo& SourceBundleInfo) const = 0;
 
 	template<typename T>
-	inline void StatsBegin(FName BundleName, T State)
+	inline void StatsBegin(const InstallBundleUtil::FContentRequestStatsKey& Key, T State)
 	{
-		RequestStats->StatsBegin(BundleName, LexToString(State));
+		RequestStats->StatsBegin(Key, LexToString(State));
 	}
 	
 	template<typename T>
-	inline void StatsEnd(FName BundleName, T State, uint64 DataSize = 0)
+	inline void StatsEnd(const InstallBundleUtil::FContentRequestStatsKey& Key, T State, uint64 DataSize = 0)
 	{
-		RequestStats->StatsEnd(BundleName, LexToString(State), DataSize);
+		RequestStats->StatsEnd(Key, LexToString(State), DataSize);
 	}
 
 	TSharedPtr<IAnalyticsProviderET> AnalyticsProvider;
diff --git a/Engine/Source/Runtime/InstallBundleManager/Private/InstallBundleUtils.cpp b/Engine/Source/Runtime/InstallBundleManager/Private/InstallBundleUtils.cpp
--- a/Engine/Source/Runtime/InstallBundleManager/Private/InstallBundleUtils.cpp
+++ b/Engine/Source/Runtime/InstallBundleManager/Private/InstallBundleUtils.cpp
@@ -389,10 +389,10 @@ namespace InstallBundleUtil
 		}
 	}
 
-	void FContentRequestStatsMap::StatsBegin(FName BundleName)
+	void FContentRequestStatsMap::StatsBegin(const FContentRequestStatsKey& Key)
 	{
-		FContentRequestStats& Stats = StatsMap.FindOrAdd(BundleName);
-		if (false == ensureAlwaysMsgf(Stats.bOpen, TEXT("StatsBegin - Stat closed for %s"), *BundleName.ToString()))
+		FContentRequestStats& Stats = StatsMap.FindOrAdd(Key);
+		if (false == ensureAlwaysMsgf(Stats.bOpen, TEXT("StatsBegin - Stat closed for %s"), *Key.GetBundleName().ToString()))
 		{
 			Stats = FContentRequestStats();
 		}
@@ -400,48 +400,48 @@ namespace InstallBundleUtil
 		Stats.StartTime = FPlatformTime::Seconds();
 	}
 
-	void FContentRequestStatsMap::StatsEnd(FName BundleName)
+	void FContentRequestStatsMap::StatsEnd(const FContentRequestStatsKey& Key)
 	{
-		FContentRequestStats& Stats = StatsMap.FindOrAdd(BundleName);
+		FContentRequestStats& Stats = StatsMap.FindOrAdd(Key);
 
-		ensureAlwaysMsgf(Stats.bOpen && Stats.StartTime > 0, TEXT("StatsEnd - Stat closed for %s"), *BundleName.ToString());
+		ensureAlwaysMsgf(Stats.bOpen && Stats.StartTime > 0, TEXT("StatsEnd - Stat closed for %s"), *Key.GetBundleName().ToString());
 		if (Stats.bOpen)
 		{
 			ensureAlwaysMsgf(
 				Algo::AllOf(Stats.StateStats, 
 					[](const TPair<FString, FContentRequestStateStats>& Pair) { return !Pair.Value.bOpen; }),
-				TEXT("StatsEnd - StateStat open for %s"), *BundleName.ToString());
+				TEXT("StatsEnd - StateStat open for %s"), *Key.GetBundleName().ToString());
 
 			Stats.EndTime = FPlatformTime::Seconds();
 			Stats.bOpen = false;
 		}
 	}
 
-	void FContentRequestStatsMap::StatsReset(FName BundleName)
+	void FContentRequestStatsMap::StatsReset(const FContentRequestStatsKey& Key)
 	{
-		if (FContentRequestStats* Stats = StatsMap.Find(BundleName))
+		if (FContentRequestStats* Stats = StatsMap.Find(Key))
 		{
-			ensureAlwaysMsgf(!Stats->bOpen, TEXT("StatsReset - Stat open for %s"), *BundleName.ToString());
+			ensureAlwaysMsgf(!Stats->bOpen, TEXT("StatsReset - Stat open for %s"), *Key.GetBundleName().ToString());
 			ensureAlwaysMsgf(
 				Algo::AllOf(Stats->StateStats,
 					[](const TPair<FString, FContentRequestStateStats>& Pair) { return !Pair.Value.bOpen; }),
-				TEXT("StatsReset - StateStat open for %s"), *BundleName.ToString());
+				TEXT("StatsReset - StateStat open for %s"), *Key.GetBundleName().ToString());
 
-			StatsMap.Remove(BundleName);
+			StatsMap.Remove(Key);
 		}
 	}
 
-	void FContentRequestStatsMap::StatsBegin(FName BundleName, const TCHAR* State)
+	void FContentRequestStatsMap::StatsBegin(const FContentRequestStatsKey& Key, const TCHAR* State)
 	{
-		FContentRequestStats& Stats = StatsMap.FindOrAdd(BundleName);
-		if (false == ensureAlwaysMsgf(Stats.bOpen, TEXT("StatsBegin - Stat closed for %s - %s"), *BundleName.ToString(), State))
+		FContentRequestStats& Stats = StatsMap.FindOrAdd(Key);
+		if (false == ensureAlwaysMsgf(Stats.bOpen, TEXT("StatsBegin - Stat closed for %s - %s"), *Key.GetBundleName().ToString(), State))
 		{
 			Stats = FContentRequestStats();
 			Stats.StartTime = FPlatformTime::Seconds();
 		}
 
 		FContentRequestStateStats& StateStats = Stats.StateStats.FindOrAdd(State);
-		if (false == ensureAlwaysMsgf(StateStats.bOpen, TEXT("StatsBegin - StateStat closed for %s - %s"), *BundleName.ToString(), State))
+		if (false == ensureAlwaysMsgf(StateStats.bOpen, TEXT("StatsBegin - StateStat closed for %s - %s"), *Key.GetBundleName().ToString(), State))
 		{
 			StateStats = FContentRequestStateStats();
 		}
@@ -449,17 +449,17 @@ namespace InstallBundleUtil
 		StateStats.StartTime = FPlatformTime::Seconds();
 	}
 
-	void FContentRequestStatsMap::StatsEnd(FName BundleName, const TCHAR* State, uint64 DataSize /*= 0*/)
+	void FContentRequestStatsMap::StatsEnd(const FContentRequestStatsKey& Key, const TCHAR* State, uint64 DataSize /*= 0*/)
 	{
-		FContentRequestStats& Stats = StatsMap.FindOrAdd(BundleName);
-		if (false == ensureAlwaysMsgf(Stats.bOpen && Stats.StartTime > 0, TEXT("StatsEnd - Stat closed for %s - %s"), *BundleName.ToString(), State))
+		FContentRequestStats& Stats = StatsMap.FindOrAdd(Key);
+		if (false == ensureAlwaysMsgf(Stats.bOpen && Stats.StartTime > 0, TEXT("StatsEnd - Stat closed for %s - %s"), *Key.GetBundleName().ToString(), State))
 		{
 			Stats = FContentRequestStats();
 			Stats.StartTime = FPlatformTime::Seconds();
 		}
 
 		FContentRequestStateStats& StateStats = Stats.StateStats.FindOrAdd(State);
-		if (ensureAlwaysMsgf(StateStats.bOpen && StateStats.StartTime > 0, TEXT("StatsEnd - StateStat closed for %s - %s"), *BundleName.ToString(), State))
+		if (ensureAlwaysMsgf(StateStats.bOpen && StateStats.StartTime > 0, TEXT("StatsEnd - StateStat closed for %s - %s"), *Key.GetBundleName().ToString(), State))
 		{
 			StateStats.EndTime = FPlatformTime::Seconds();
 			StateStats.DataSize = DataSize;
diff --git a/Engine/Source/Runtime/InstallBundleManager/Public/InstallBundleSourceInterface.h b/Engine/Source/Runtime/InstallBundleManager/Public/InstallBundleSourceInterface.h
--- a/Engine/Source/Runtime/InstallBundleManager/Public/InstallBundleSourceInterface.h
+++ b/Engine/Source/Runtime/InstallBundleManager/Public/InstallBundleSourceInterface.h
@@ -83,6 +83,7 @@ public:
 		FInstallBundlePausedDelegate PausedCallback;
 		FInstallBundleCompleteDelegate CompleteCallback;
 		TSharedPtr<InstallBundleUtil::FContentRequestSharedContext> RequestSharedContext;
+		void* StatsCookie = nullptr;
 	};
 
 	// Updates content on disk if necessary
diff --git a/Engine/Source/Runtime/InstallBundleManager/Public/InstallBundleUtils.h b/Engine/Source/Runtime/InstallBundleManager/Public/InstallBundleUtils.h
--- a/Engine/Source/Runtime/InstallBundleManager/Public/InstallBundleUtils.h
+++ b/Engine/Source/Runtime/InstallBundleManager/Public/InstallBundleUtils.h
@@ -248,19 +248,66 @@ namespace InstallBundleUtil
 		}
 	};
 
+	class FContentRequestStatsKey
+	{
+	private:
+		FName BundleName;
+		void* Cookie = nullptr;
+
+	public:
+		FContentRequestStatsKey(FName InBundleName, void* InCookie)
+			: BundleName(InBundleName), Cookie(InCookie)
+		{}
+
+		FName GetBundleName() const { return BundleName; }
+		void* GetCookie() const { return Cookie; }
+
+		[[nodiscard]] bool operator==(const FContentRequestStatsKey& Other) const
+		{
+			return BundleName == Other.BundleName && Cookie == Other.Cookie;
+		}
+
+		[[nodiscard]] friend uint32 GetTypeHash(const FContentRequestStatsKey& Key)
+		{
+			return HashCombineFast(GetTypeHash(Key.GetBundleName()), GetTypeHash(Key.GetCookie()));
+		}
+	};
+
 	class FContentRequestStatsMap
 	{
 	private:
-		TMap<FName, InstallBundleUtil::FContentRequestStats> StatsMap;
+		TMap<FContentRequestStatsKey, FContentRequestStats> StatsMap;
 
 	public:
-		INSTALLBUNDLEMANAGER_API void StatsBegin(FName BundleName);
-		INSTALLBUNDLEMANAGER_API void StatsEnd(FName BundleName);
-		INSTALLBUNDLEMANAGER_API void StatsBegin(FName BundleName, const TCHAR* State);
-		INSTALLBUNDLEMANAGER_API void StatsEnd(FName BundleName, const TCHAR* State, uint64 DataSize = 0);
-		INSTALLBUNDLEMANAGER_API void StatsReset(FName BundleName);
+		INSTALLBUNDLEMANAGER_API void StatsBegin(const FContentRequestStatsKey& Key);
+		INSTALLBUNDLEMANAGER_API void StatsEnd(const FContentRequestStatsKey& Key);
+		INSTALLBUNDLEMANAGER_API void StatsBegin(const FContentRequestStatsKey& Key, const TCHAR* State);
+		INSTALLBUNDLEMANAGER_API void StatsEnd(const FContentRequestStatsKey& Key, const TCHAR* State, uint64 DataSize = 0);
+		INSTALLBUNDLEMANAGER_API void StatsReset(const FContentRequestStatsKey& Key);
+
+		const FContentRequestStats* FindStats(const FContentRequestStatsKey& Key) const
+		{
+			return StatsMap.Find(Key);
+		}
+
+		const FContentRequestStats& FindStatsChecked(const FContentRequestStatsKey& Key) const
+		{
+			return StatsMap.FindChecked(Key);
+		}
 
-		const TMap<FName, InstallBundleUtil::FContentRequestStats>& GetMap() { return StatsMap; }
+		UE_DEPRECATED(5.7, "Use FContentRequestStatsKey version")
+		void StatsBegin(FName BundleName) { StatsBegin(FContentRequestStatsKey(BundleName, nullptr)); }
+		UE_DEPRECATED(5.7, "Use FContentRequestStatsKey version")
+		void StatsEnd(FName BundleName) { StatsEnd(FContentRequestStatsKey(BundleName, nullptr)); }
+		UE_DEPRECATED(5.7, "Use FContentRequestStatsKey version")
+		void StatsBegin(FName BundleName, const TCHAR* State) { StatsBegin(FContentRequestStatsKey(BundleName, nullptr)); }
+		UE_DEPRECATED(5.7, "Use FContentRequestStatsKey version")
+		void StatsEnd(FName BundleName, const TCHAR* State, uint64 DataSize = 0) { StatsEnd(FContentRequestStatsKey(BundleName, nullptr)); }
+		UE_DEPRECATED(5.7, "Use FContentRequestStatsKey version")
+		void StatsReset(FName BundleName) { StatsReset(FContentRequestStatsKey(BundleName, nullptr)); }
+
+		UE_DEPRECATED(5.7, "Use FindStats and FindStatsChecked")
+		const TMap<FContentRequestStatsKey, FContentRequestStats>& GetMap() { return StatsMap; }
 	};
 
 	namespace PersistentStats
@@ -653,4 +700,4 @@ namespace InstallBundleUtil
 		TMap<FInstallBundleSourceType, TUniquePtr<IBundleSourceContentRequestSharedContext>> BundleSourceSharedContext;
 	};
 	using FContentRequestSharedContextPtr = TSharedPtr<FContentRequestSharedContext>;
-}
+} // InstallBundleUtil
