Deprecate InstallBundleSourcePlayGo

diff --git a/Engine/Plugins/Experimental/DefaultInstallBundleManager/Source/Private/InstallBundleManagerUtil.cpp b/Engine/Plugins/Experimental/DefaultInstallBundleManager/Source/Private/InstallBundleManagerUtil.cpp
--- a/Engine/Plugins/Experimental/DefaultInstallBundleManager/Source/Private/InstallBundleManagerUtil.cpp
+++ b/Engine/Plugins/Experimental/DefaultInstallBundleManager/Source/Private/InstallBundleManagerUtil.cpp
@@ -38,7 +38,10 @@ namespace InstallBundleManagerUtil
 #if WITH_PLATFORM_INSTALL_BUNDLE_SOURCE
 		if (Type.GetName() == TEXTVIEW("Platform"))
 		{
-			return MakePlatformBundleSource();
+			if (TSharedPtr<IInstallBundleSource> PlatformSource = MakePlatformBundleSource())
+			{
+				return PlatformSource;
+			}
 		}
 #endif // WITH_PLATFORM_INSTALL_BUNDLE_SOURCE
 
