Install Bundle Manager

Improve ensure error when a bundle sources sends OnDemandMountArgs unexpectedly.

diff --git a/Engine/Plugins/Experimental/DefaultInstallBundleManager/Source/Private/DefaultInstallBundleManager.cpp b/Engine/Plugins/Experimental/DefaultInstallBundleManager/Source/Private/DefaultInstallBundleManager.cpp
--- a/Engine/Plugins/Experimental/DefaultInstallBundleManager/Source/Private/DefaultInstallBundleManager.cpp
+++ b/Engine/Plugins/Experimental/DefaultInstallBundleManager/Source/Private/DefaultInstallBundleManager.cpp
@@ -1341,8 +1341,16 @@ void FDefaultInstallBundleManager::UpdateBundleSourceComplete(TSharedRef<IInstal
 		return A.Key > B.Key;
 	});
 
-	ensureMsgf(!Request->OnDemandMountArgs.IsEmpty() == BundleInfo.bContainsIoStoreOnDemandTocs,
-		TEXT("OnDemandMountArgs should only be provided by bundle sources that set bContainsIoStoreOnDemandTocs"));
+	if (BundleInfo.bContainsIoStoreOnDemandTocs)
+	{
+		ensureMsgf(!Request->OnDemandMountArgs.IsEmpty(),
+			TEXT("Expected OnDemandMountArgs for bundle %s, check that your bundles sources are configured to match the data being installed. This could also be caused by improper ini settings or CVars."), *Request->BundleName.ToString());
+	}
+	else
+	{
+		ensureMsgf(Request->OnDemandMountArgs.IsEmpty(),
+			TEXT("Unexpected OnDemandMountArgs for bundle %s, check that your bundles sources are configured to match the data being installed. This could also be caused by improper ini settings or CVars."), *Request->BundleName.ToString());
+	}
 
 	// If there are no content paths, its likely this is a chunk that doesn't exist on the current platfrom, so set bContainsChunks true.
 	// This is a corner case but as far as I know there is no other situation that would allow for an empty bundle that was not chunked.
