Make IBM Analytics provider create a normal SessionID so it can be given to the main client analytics provider later and look normal.
#rnx
#tests Ran local client and observed telemetry signatures.

diff --git a/Engine/Plugins/Experimental/DefaultInstallBundleManager/Source/Private/DefaultInstallBundleManager.cpp b/Engine/Plugins/Experimental/DefaultInstallBundleManager/Source/Private/DefaultInstallBundleManager.cpp
--- a/Engine/Plugins/Experimental/DefaultInstallBundleManager/Source/Private/DefaultInstallBundleManager.cpp
+++ b/Engine/Plugins/Experimental/DefaultInstallBundleManager/Source/Private/DefaultInstallBundleManager.cpp
@@ -5176,36 +5176,9 @@ void FDefaultInstallBundleManager::AsyncInit_SetUpdateBundleInfoCallback()
 
 void FDefaultInstallBundleManager::AsyncInit_CreateAnalyticsSession()
 {
-	LOG_INSTALL_BUNDLE_MAN(Display, TEXT("Setting up analytics"));
-
-	if (AnalyticsProvider)
-	{
-		// Pick the latest content version from all sources
-		FString ContentVersion;
-		int64 MaxVersionCL = -1;
-		for (const TPair<FInstallBundleSourceType, TSharedPtr<IInstallBundleSource>>& Pair : BundleSources)
-		{
-			const TSharedPtr<IInstallBundleSource>& Source = Pair.Value;
-			FString SourceContentVersion = Source->GetContentVersion();
-			int64 SourceContentVersionCL = -1;
-
-			FString CLString;
-			if (SourceContentVersion.Split(TEXT("CL-"), nullptr, &CLString, ESearchCase::CaseSensitive))
-			{
-				CLString.Split(TEXT("-"), &CLString, nullptr, ESearchCase::CaseSensitive);
-				LexTryParseString(SourceContentVersionCL, *CLString);
-			}
-
-			if (MaxVersionCL == -1 || MaxVersionCL < SourceContentVersionCL)
-			{
-				MaxVersionCL = SourceContentVersionCL;
-				ContentVersion = SourceContentVersion;
-			}
-		}
-		ensureAlways(MaxVersionCL != -1);
-
-		AnalyticsProvider->SetSessionID(FString(TEXT("IBMInstallSession-")) + FGuid::NewGuid().ToString() + TEXT("-") + ContentVersion);
-	}
+	// no longer do this. We use a default SessionID GUID now when creating this analytics provider, 
+	// and no longer need to wait until we have bundle sources to set the session id.
+	// We could remove this step, but it would make a more disruptive change.
 	InitStepResult = EAsyncInitStepResult::Done;
 }
 
