Fix to race condition associated with ordering of async callbacks. If an empty bundle arrived last, it would stomp the free size with zero, causing an incorrect out of disk space check. Since multiple bundle requests get queued, one could succeed and start the install while the other failed.

diff --git a/Engine/Plugins/Experimental/DefaultInstallBundleManager/Source/Private/DefaultInstallBundleManager.cpp b/Engine/Plugins/Experimental/DefaultInstallBundleManager/Source/Private/DefaultInstallBundleManager.cpp
--- a/Engine/Plugins/Experimental/DefaultInstallBundleManager/Source/Private/DefaultInstallBundleManager.cpp
+++ b/Engine/Plugins/Experimental/DefaultInstallBundleManager/Source/Private/DefaultInstallBundleManager.cpp
@@ -545,7 +545,15 @@ void FDefaultInstallBundleManager::TickGetContentState()
 				}
 
 				State.ContentSize += SourceState.ContentSize;
-				State.FreeSpace = State.FreeSpace ? FMath::Min(State.FreeSpace, SourceState.FreeSpace) : SourceState.FreeSpace;
+				
+				if (SourceState.IndividualBundleStates.Num()) // FreeSpace is only valid if there's content.
+				{
+					// Free space should be usually be roughly the same, but Apple platforms can dynamically resize the partition
+					// on demand, so we use Max() to ensure we catch the larger size if they've moved on us. It's possible we got
+					// less due to another app writing to disk, but in such a case we're already going to be reliant on the in-progress
+					// out of disk space checks.
+					State.FreeSpace = FMath::Max(State.FreeSpace, SourceState.FreeSpace);
+				}
 			}
 
 			for (const TPair<FName, FInstallBundleContentState>& Pair : State.IndividualBundleStates)
