Combine InstallBundleManager.BundleSourceBPS events into InstallBundleManager

#rnx

diff --git a/Engine/Plugins/Experimental/DefaultInstallBundleManager/Source/Private/DefaultInstallBundleManager.cpp b/Engine/Plugins/Experimental/DefaultInstallBundleManager/Source/Private/DefaultInstallBundleManager.cpp
--- a/Engine/Plugins/Experimental/DefaultInstallBundleManager/Source/Private/DefaultInstallBundleManager.cpp
+++ b/Engine/Plugins/Experimental/DefaultInstallBundleManager/Source/Private/DefaultInstallBundleManager.cpp
@@ -2171,8 +2171,18 @@ void FDefaultInstallBundleManager::TickContentRequests()
 			break;
 
 		AddRequestToInitialBatch(Request);
-
-		if (Request->bShouldSendAnalytics)
+		
+		// If the source supports BundleRequestStartedDelegate (most likely BPS), will skip analytics here, will do it during delegate callback
+		bool bSupportsBundleRequestStartedDelegate = false;
+		TArray<TSharedPtr<IInstallBundleSource>> EnabledBundleSources = GetEnabledBundleSourcesForRequest(Request);
+		for (const TSharedPtr<IInstallBundleSource>& Source : EnabledBundleSources)
+		{
+			if (Source->SupportsBundleRequestEventDelegates())
+			{
+				bSupportsBundleRequestStartedDelegate = true;
+			}
+		}
+		if (!bSupportsBundleRequestStartedDelegate)
 		{
 			InstallBundleManagerAnalytics::FireEvent_BundleRequestStarted(AnalyticsProvider.Get(), Request->BundleName.ToString());
 		}
@@ -2281,11 +2291,26 @@ void FDefaultInstallBundleManager::TickContentRequests()
 
 				if (Request->bShouldSendAnalytics || Request->Result != EInstallBundleResult::OK)
 				{
-					InstallBundleManagerAnalytics::FireEvent_BundleRequestComplete(AnalyticsProvider.Get(),
-						Request->BundleName.ToString(),
-						Request->bContentWasInstalled,
-						LexToString(Request->Result),
-						StatsMap->FindStatsChecked(Request->GetStatsKey()));
+					// check if source delegate has provided data from BPS, if so send a combined event
+					auto Data = BundleRequestCompleteInfoMap.Find(Request->BundleName.ToString());
+					if (Data)
+					{
+						InstallBundleManagerAnalytics::FireEvent_BundleRequestCompleteDetailed(AnalyticsProvider.Get(),
+							Request->BundleName.ToString(),
+							Request->bContentWasInstalled,
+							LexToString(Request->Result),
+							StatsMap->FindStatsChecked(Request->GetStatsKey()),
+							*Data);
+						BundleRequestCompleteInfoMap.Remove(Request->BundleName.ToString());
+					}
+					else
+					{
+						InstallBundleManagerAnalytics::FireEvent_BundleRequestComplete(AnalyticsProvider.Get(),
+							Request->BundleName.ToString(),
+							Request->bContentWasInstalled,
+							LexToString(Request->Result),
+							StatsMap->FindStatsChecked(Request->GetStatsKey()));
+					}
 				}
 
 				LOG_INSTALL_BUNDLE_MAN_OVERRIDE(Request->LogVerbosityOverride, Display, TEXT("Removing Request %s"), *Request->BundleName.ToString());
@@ -4648,6 +4673,34 @@ EInstallBundleManagerInitResult FDefaultInstallBundleManager::Init_TryCreateBund
 							}
 						}));
 				}
+				if (Source->SupportsBundleRequestEventDelegates())
+				{
+					Source->SetBundleRequestStartedDelegate(FInstallBundleRequestStartedDelegate::CreateLambda([this, WeakSource = Source.ToWeakPtr()](const FString& BundleName, const FString& ManifestName, const FString& ManifestPath, const FString& ManifestVersion)
+					{
+						for (EContentRequestBatch b : TEnumRange<EContentRequestBatch>())
+						{
+							for (const FContentRequestRef& Request : ContentRequests[b])
+							{
+								if (Request->bShouldSendAnalytics && Request->BundleName.ToString().Equals(BundleName))
+								{
+									InstallBundleManagerAnalytics::FireEvent_BundleRequestStartedDetailed(AnalyticsProvider.Get(),
+										Request->BundleName.ToString(),
+										ManifestName,
+										ManifestPath,
+										ManifestVersion);
+									return;
+								}
+							}
+						}
+					}));
+					Source->SetBundleRequestCompletedDelegate(FInstallBundleRequestCompletedDelegate::CreateLambda([this, WeakSource = Source.ToWeakPtr()]
+																												   (const FString& BundleName,
+																													const FBundleRequestCompleteInfo& AnalyticInfo)
+					{
+						// Save for later complete event firing
+						BundleRequestCompleteInfoMap.Emplace(BundleName, AnalyticInfo);
+					}));
+				}
 				BundleSources.Add(SourceType, Source);
 				if (OutNewSources)
 				{
diff --git a/Engine/Plugins/Experimental/DefaultInstallBundleManager/Source/Private/InstallBundleManagerUtil.cpp b/Engine/Plugins/Experimental/DefaultInstallBundleManager/Source/Private/InstallBundleManagerUtil.cpp
--- a/Engine/Plugins/Experimental/DefaultInstallBundleManager/Source/Private/InstallBundleManagerUtil.cpp
+++ b/Engine/Plugins/Experimental/DefaultInstallBundleManager/Source/Private/InstallBundleManagerUtil.cpp
@@ -740,6 +740,23 @@ namespace InstallBundleManagerAnalytics
 			TEXT("BundleName"), BundleName));
 	}
 
+	void FireEvent_BundleRequestStartedDetailed(IAnalyticsProviderET* AnalyticsProvider,
+		const FString& BundleName,
+		const FString& PreviousManifest, const FString& CurrentManifest,
+		const FString& InstallManifestVersion)
+	{
+		if (AnalyticsProvider == nullptr || InstallBundleUtil::FInstallBundleSuppressAnalytics::IsEnabled())
+		{
+			return;
+		}
+
+		AnalyticsProvider->RecordEvent(TEXT("InstallBundleManager.BundleRequestStartedDetailed"), MakeAnalyticsEventAttributeArray(
+			TEXT("BundleName"), BundleName,
+			TEXT("PreviousManfest"), PreviousManifest,
+			TEXT("CurrentManifest"), CurrentManifest,
+			TEXT("CurrentVersion"), InstallManifestVersion));
+	}
+
 	void FireEvent_BundleRequestComplete(IAnalyticsProviderET* AnalyticsProvider, 
 		const FString& BundleName, 
 		bool bDidInstall, 
@@ -767,6 +784,43 @@ namespace InstallBundleManagerAnalytics
 		AnalyticsProvider->RecordEvent(TEXT("InstallBundleManager.BundleRequestComplete"), MoveTemp(Attributes));
 	}
 
+	void FireEvent_BundleRequestCompleteDetailed(IAnalyticsProviderET* AnalyticsProvider, 
+		const FString& BundleName, 
+		bool bDidInstall, 
+		const FString& Result, 
+		const InstallBundleUtil::FContentRequestStats& TimingStats,
+		const FBundleRequestCompleteInfo& RequestCompleteInfo)
+	{
+		if (AnalyticsProvider == nullptr || InstallBundleUtil::FInstallBundleSuppressAnalytics::IsEnabled())
+		{
+			return;
+		}
+
+		TArray<FAnalyticsEventAttribute> Attributes = MakeAnalyticsEventAttributeArray(
+			TEXT("BundleName"), BundleName,
+			TEXT("DidInstall"), bDidInstall,
+			TEXT("Result"), Result,
+			TEXT("Total_Time"), TimingStats.GetElapsedTime());
+
+		Attributes.Reserve(Attributes.Num() + TimingStats.StateStats.Num());
+
+		for (const TPair<FString, InstallBundleUtil::FContentRequestStateStats>& Pair : TimingStats.StateStats)
+		{
+			Attributes.Add(FAnalyticsEventAttribute(Pair.Key + TEXT("_Time"), Pair.Value.GetElapsedTime()));
+		}
+		
+		Attributes.Add(FAnalyticsEventAttribute(TEXT("PreviousManfest"), RequestCompleteInfo.PreviousManifest));
+		Attributes.Add(FAnalyticsEventAttribute(TEXT("CurrentManifest"), RequestCompleteInfo.CurrentManifest));
+		Attributes.Add(FAnalyticsEventAttribute(TEXT("PreviousVersion"), RequestCompleteInfo.OldManifestVersion));
+		Attributes.Add(FAnalyticsEventAttribute(TEXT("CurrentVersion"), RequestCompleteInfo.InstallManifestVersion));
+		Attributes.Add(FAnalyticsEventAttribute(TEXT("PreviousVersionTimeStamp"), RequestCompleteInfo.OldVersionTimeStamp));
+		Attributes.Add(FAnalyticsEventAttribute(TEXT("TotalDownloadedBytes"), RequestCompleteInfo.TotalDownloadedBytes));
+		Attributes.Add(FAnalyticsEventAttribute(TEXT("EstimatedFullDownloadBytes"), RequestCompleteInfo.EstimatedFullDownloadBytes));
+		Attributes.Add(FAnalyticsEventAttribute(TEXT("BPTErrorCode"), RequestCompleteInfo.BPTErrorCode));
+
+		AnalyticsProvider->RecordEvent(TEXT("InstallBundleManager.BundleRequestCompleteDetailed"), MoveTemp(Attributes));
+	}
+
 	void FireEvent_BundleReleaseRequestStarted(IAnalyticsProviderET* AnalyticsProvider, 
 		const FString& BundleName, 
 		bool bRemoveFilesIfPossible,
diff --git a/Engine/Plugins/Experimental/DefaultInstallBundleManager/Source/Public/DefaultInstallBundleManager.h b/Engine/Plugins/Experimental/DefaultInstallBundleManager/Source/Public/DefaultInstallBundleManager.h
--- a/Engine/Plugins/Experimental/DefaultInstallBundleManager/Source/Public/DefaultInstallBundleManager.h
+++ b/Engine/Plugins/Experimental/DefaultInstallBundleManager/Source/Public/DefaultInstallBundleManager.h
@@ -831,6 +831,7 @@ protected:
 	// Analytics
 	TSharedPtr<IAnalyticsProviderET> AnalyticsProvider;
 	TSharedRef<InstallBundleUtil::FContentRequestStatsMap> StatsMap;
+	TMap<FString, FBundleRequestCompleteInfo> BundleRequestCompleteInfoMap;	// BundleName as key
 };
 
 
diff --git a/Engine/Plugins/Experimental/DefaultInstallBundleManager/Source/Public/InstallBundleManagerUtil.h b/Engine/Plugins/Experimental/DefaultInstallBundleManager/Source/Public/InstallBundleManagerUtil.h
--- a/Engine/Plugins/Experimental/DefaultInstallBundleManager/Source/Public/InstallBundleManagerUtil.h
+++ b/Engine/Plugins/Experimental/DefaultInstallBundleManager/Source/Public/InstallBundleManagerUtil.h
@@ -313,6 +313,24 @@ namespace InstallBundleManagerAnalytics
 	DEFAULTINSTALLBUNDLEMANAGER_API void FireEvent_BundleRequestStarted(IAnalyticsProviderET* AnalyticsProvider,
 		const FString& BundleName);
 
+	/**
+	 * @EventName InstallBundleManager.BundleRequestStartedDetailed
+	 *
+	 * @Trigger Fired when a bundle request is started, including more details from implementation
+	 *
+	 * @Type Client
+	 *
+	 * @EventParam BundleName (string) String Name of bundle being installed.
+	 * @EventParam PreviousManfest (string) Manifest version information for currently installed version
+	 * @EventParam CurrentManifest (string) Manifest version information for version we are going to patch to
+	 * @EventParam InstallManifestVersion (string) Content version information for version we are going to patch to 
+	 *
+	 */
+	DEFAULTINSTALLBUNDLEMANAGER_API void FireEvent_BundleRequestStartedDetailed(IAnalyticsProviderET* AnalyticsProvider,
+		const FString& BundleName,
+		const FString& PreviousManifest, const FString& CurrentManifest,
+		const FString& InstallManifestVersion);
+
 	/**
 	 * @EventName InstallBundleManager.BundleRequestComplete
 	 *
@@ -331,6 +349,27 @@ namespace InstallBundleManagerAnalytics
 		const FString& Result,
 		const InstallBundleUtil::FContentRequestStats& TimingStats);
 
+	/**
+	 * @EventName InstallBundleManager.BundleRequestCompleteDetailed
+	 *
+	 * @Trigger Fired after an install bundle request is completed, including more details from implementation
+	 *
+	 * @Type Client
+	 *
+	 * @EventParam BundleName (string) String Name of bundle being installed.
+	 * @EventParam DidInstall (bool) True if the bundle was installed successfully
+	 * @EventParam Result (string) Result of the request
+	 * @EventParam TimingStates (double) A variable number of stats based on what bundle manager steps were run. Fields end in "_Time"
+	 * @EventParam RequestCompleteInfo (FBundleRequestCompleteInfo) Detailed info from implementation
+	 */
+	DEFAULTINSTALLBUNDLEMANAGER_API void FireEvent_BundleRequestCompleteDetailed(IAnalyticsProviderET* AnalyticsProvider,
+		const FString& BundleName,
+		bool bDidInstall,
+		const FString& Result,
+		const InstallBundleUtil::FContentRequestStats& TimingStats,
+		const FBundleRequestCompleteInfo& RequestCompleteInfo);
+
+
 	/**
 	 * @EventName InstallBundleManager.BundleReleaseRequestStarted
 	 *
diff --git a/Engine/Source/Runtime/InstallBundleManager/Public/InstallBundleSourceInterface.h b/Engine/Source/Runtime/InstallBundleManager/Public/InstallBundleSourceInterface.h
--- a/Engine/Source/Runtime/InstallBundleManager/Public/InstallBundleSourceInterface.h
+++ b/Engine/Source/Runtime/InstallBundleManager/Public/InstallBundleSourceInterface.h
@@ -21,6 +21,9 @@ DECLARE_DELEGATE_TwoParams(FInstallBundleRemovedDelegate, TSharedRef<IInstallBun
 
 DECLARE_DELEGATE_OneParam(FInstallBundleReportingDelegate, const FInstallBundleReport&);
 
+DECLARE_DELEGATE_FourParams(FInstallBundleRequestStartedDelegate, const FString& /*BundleName*/, const FString& /*ManifestName*/, const FString& /*ManifestPath*/, const FString& /*ManifestVersion*/);
+DECLARE_DELEGATE_TwoParams(FInstallBundleRequestCompletedDelegate, const FString& /*BundleName*/, const FBundleRequestCompleteInfo& /*AnalyticInfo*/);
+
 DECLARE_DELEGATE_TwoParams(FInstallBundleSourceContentPatchResultDelegate, TSharedRef<IInstallBundleSource> /*Source*/, bool /*bContentPatchRequired*/);
 
 class IInstallBundleSource : public TSharedFromThis<IInstallBundleSource>
@@ -153,6 +156,10 @@ public:
 
 	virtual bool SupportsReportingDelegate() const { return false; }
 	virtual void SetReportingDelegate(FInstallBundleReportingDelegate Delegate) {}
+	
+	virtual bool SupportsBundleRequestEventDelegates() const { return false; }
+	virtual void SetBundleRequestStartedDelegate(FInstallBundleRequestStartedDelegate Delegate) {}
+	virtual void SetBundleRequestCompletedDelegate(FInstallBundleRequestCompletedDelegate Delegate) {}
 
 private:
 	TArray<FBuildInstallerStat> EmptyStats;
diff --git a/Engine/Source/Runtime/InstallBundleManager/Public/InstallBundleTypes.h b/Engine/Source/Runtime/InstallBundleManager/Public/InstallBundleTypes.h
--- a/Engine/Source/Runtime/InstallBundleManager/Public/InstallBundleTypes.h
+++ b/Engine/Source/Runtime/InstallBundleManager/Public/InstallBundleTypes.h
@@ -495,3 +495,16 @@ enum class UE_DEPRECATED(5.7, "Use EInstallBundleCacheStatsFlags") EInstallBundl
 	Default,
 	CSV
 };
+
+struct FBundleRequestCompleteInfo
+{
+	FString PreviousManifest;
+	FString CurrentManifest;
+	FString OldManifestVersion;
+	FString InstallManifestVersion;
+	FString OldVersionTimeStamp;
+	uint64 TotalDownloadedBytes;
+	uint64 EstimatedFullDownloadBytes;
+	FString Result;
+	FString BPTErrorCode;
+};
