Fixing minor issues for analytic event merge.

#rnx

diff --git a/Engine/Plugins/Experimental/DefaultInstallBundleManager/Source/Private/DefaultInstallBundleManager.cpp b/Engine/Plugins/Experimental/DefaultInstallBundleManager/Source/Private/DefaultInstallBundleManager.cpp
--- a/Engine/Plugins/Experimental/DefaultInstallBundleManager/Source/Private/DefaultInstallBundleManager.cpp
+++ b/Engine/Plugins/Experimental/DefaultInstallBundleManager/Source/Private/DefaultInstallBundleManager.cpp
@@ -2292,7 +2292,7 @@ void FDefaultInstallBundleManager::TickContentRequests()
 				if (Request->bShouldSendAnalytics || Request->Result != EInstallBundleResult::OK)
 				{
 					// check if source delegate has provided data from BPS, if so send a combined event
-					auto Data = BundleRequestCompleteInfoMap.Find(Request->BundleName.ToString());
+					auto Data = BundleRequestCompleteInfoMap.Find(Request->BundleName);
 					if (Data)
 					{
 						InstallBundleManagerAnalytics::FireEvent_BundleRequestCompleteDetailed(AnalyticsProvider.Get(),
@@ -2301,7 +2301,7 @@ void FDefaultInstallBundleManager::TickContentRequests()
 							LexToString(Request->Result),
 							StatsMap->FindStatsChecked(Request->GetStatsKey()),
 							*Data);
-						BundleRequestCompleteInfoMap.Remove(Request->BundleName.ToString());
+						BundleRequestCompleteInfoMap.Remove(Request->BundleName);
 					}
 					else
 					{
@@ -4675,13 +4675,17 @@ EInstallBundleManagerInitResult FDefaultInstallBundleManager::Init_TryCreateBund
 				}
 				if (Source->SupportsBundleRequestEventDelegates())
 				{
-					Source->SetBundleRequestStartedDelegate(FInstallBundleRequestStartedDelegate::CreateLambda([this, WeakSource = Source.ToWeakPtr()](const FString& BundleName, const FString& ManifestName, const FString& ManifestPath, const FString& ManifestVersion)
+					Source->SetBundleRequestStartedDelegate(FInstallBundleRequestStartedDelegate::CreateLambda([this, WeakSource = Source.ToWeakPtr()]
+																											   (const FName& BundleName,
+																												const FString& ManifestName,
+																												const FString& ManifestPath,
+																												const FString& ManifestVersion)
 					{
 						for (EContentRequestBatch b : TEnumRange<EContentRequestBatch>())
 						{
 							for (const FContentRequestRef& Request : ContentRequests[b])
 							{
-								if (Request->bShouldSendAnalytics && Request->BundleName.ToString().Equals(BundleName))
+								if (Request->bShouldSendAnalytics && Request->BundleName.IsEqual(BundleName))
 								{
 									InstallBundleManagerAnalytics::FireEvent_BundleRequestStartedDetailed(AnalyticsProvider.Get(),
 										Request->BundleName.ToString(),
@@ -4694,7 +4698,7 @@ EInstallBundleManagerInitResult FDefaultInstallBundleManager::Init_TryCreateBund
 						}
 					}));
 					Source->SetBundleRequestCompletedDelegate(FInstallBundleRequestCompletedDelegate::CreateLambda([this, WeakSource = Source.ToWeakPtr()]
-																												   (const FString& BundleName,
+																												   (const FName& BundleName,
 																													const FBundleRequestCompleteInfo& AnalyticInfo)
 					{
 						// Save for later complete event firing
diff --git a/Engine/Plugins/Experimental/DefaultInstallBundleManager/Source/Private/InstallBundleManagerUtil.cpp b/Engine/Plugins/Experimental/DefaultInstallBundleManager/Source/Private/InstallBundleManagerUtil.cpp
--- a/Engine/Plugins/Experimental/DefaultInstallBundleManager/Source/Private/InstallBundleManagerUtil.cpp
+++ b/Engine/Plugins/Experimental/DefaultInstallBundleManager/Source/Private/InstallBundleManagerUtil.cpp
@@ -743,7 +743,7 @@ namespace InstallBundleManagerAnalytics
 	void FireEvent_BundleRequestStartedDetailed(IAnalyticsProviderET* AnalyticsProvider,
 		const FString& BundleName,
 		const FString& PreviousManifest, const FString& CurrentManifest,
-		const FString& InstallManifestVersion)
+		const FString& CurrentVersion)
 	{
 		if (AnalyticsProvider == nullptr || InstallBundleUtil::FInstallBundleSuppressAnalytics::IsEnabled())
 		{
@@ -754,7 +754,7 @@ namespace InstallBundleManagerAnalytics
 			TEXT("BundleName"), BundleName,
 			TEXT("PreviousManfest"), PreviousManifest,
 			TEXT("CurrentManifest"), CurrentManifest,
-			TEXT("CurrentVersion"), InstallManifestVersion));
+			TEXT("CurrentVersion"), CurrentVersion));
 	}
 
 	void FireEvent_BundleRequestComplete(IAnalyticsProviderET* AnalyticsProvider, 
diff --git a/Engine/Plugins/Experimental/DefaultInstallBundleManager/Source/Public/DefaultInstallBundleManager.h b/Engine/Plugins/Experimental/DefaultInstallBundleManager/Source/Public/DefaultInstallBundleManager.h
--- a/Engine/Plugins/Experimental/DefaultInstallBundleManager/Source/Public/DefaultInstallBundleManager.h
+++ b/Engine/Plugins/Experimental/DefaultInstallBundleManager/Source/Public/DefaultInstallBundleManager.h
@@ -831,7 +831,7 @@ protected:
 	// Analytics
 	TSharedPtr<IAnalyticsProviderET> AnalyticsProvider;
 	TSharedRef<InstallBundleUtil::FContentRequestStatsMap> StatsMap;
-	TMap<FString, FBundleRequestCompleteInfo> BundleRequestCompleteInfoMap;	// BundleName as key
+	TMap<FName, FBundleRequestCompleteInfo> BundleRequestCompleteInfoMap;	// BundleName as key
 };
 
 
diff --git a/Engine/Plugins/Experimental/DefaultInstallBundleManager/Source/Public/InstallBundleManagerUtil.h b/Engine/Plugins/Experimental/DefaultInstallBundleManager/Source/Public/InstallBundleManagerUtil.h
--- a/Engine/Plugins/Experimental/DefaultInstallBundleManager/Source/Public/InstallBundleManagerUtil.h
+++ b/Engine/Plugins/Experimental/DefaultInstallBundleManager/Source/Public/InstallBundleManagerUtil.h
@@ -321,15 +321,15 @@ namespace InstallBundleManagerAnalytics
 	 * @Type Client
 	 *
 	 * @EventParam BundleName (string) String Name of bundle being installed.
-	 * @EventParam PreviousManfest (string) Manifest version information for currently installed version
+	 * @EventParam PreviousManifest (string) Manifest version information for currently installed version
 	 * @EventParam CurrentManifest (string) Manifest version information for version we are going to patch to
-	 * @EventParam InstallManifestVersion (string) Content version information for version we are going to patch to 
+	 * @EventParam CurrentVersion (string) Content version information for version we are going to patch to
 	 *
 	 */
 	DEFAULTINSTALLBUNDLEMANAGER_API void FireEvent_BundleRequestStartedDetailed(IAnalyticsProviderET* AnalyticsProvider,
 		const FString& BundleName,
 		const FString& PreviousManifest, const FString& CurrentManifest,
-		const FString& InstallManifestVersion);
+		const FString& CurrentVersion);
 
 	/**
 	 * @EventName InstallBundleManager.BundleRequestComplete
@@ -360,7 +360,14 @@ namespace InstallBundleManagerAnalytics
 	 * @EventParam DidInstall (bool) True if the bundle was installed successfully
 	 * @EventParam Result (string) Result of the request
 	 * @EventParam TimingStates (double) A variable number of stats based on what bundle manager steps were run. Fields end in "_Time"
-	 * @EventParam RequestCompleteInfo (FBundleRequestCompleteInfo) Detailed info from implementation
+	 * @EventParam PreviousManfest (string) Manifest version information for currently installed version
+	 * @EventParam CurrentManifest (string) Manifest version information for version we are going to patch to
+	 * @EventParam PreviousVersion (string) Content version information for currently installed version
+	 * @EventParam CurrentVersion (string) Content version information for version we are going to patch to
+	 * @EventParam PreviousVersionTimeStamp (string) Timestamp of when the previous version as installed
+	 * @EventParam TotalDownloadedBytes (uint64) Total bytes downloaded
+	 * @EventParam EstimatedFullDownloadBytes (uint64) Estimated bytes that needed download
+	 * @EventParam BPTErrorCode (string) Any BPT Error
 	 */
 	DEFAULTINSTALLBUNDLEMANAGER_API void FireEvent_BundleRequestCompleteDetailed(IAnalyticsProviderET* AnalyticsProvider,
 		const FString& BundleName,
diff --git a/Engine/Source/Runtime/InstallBundleManager/Public/InstallBundleSourceInterface.h b/Engine/Source/Runtime/InstallBundleManager/Public/InstallBundleSourceInterface.h
--- a/Engine/Source/Runtime/InstallBundleManager/Public/InstallBundleSourceInterface.h
+++ b/Engine/Source/Runtime/InstallBundleManager/Public/InstallBundleSourceInterface.h
@@ -21,8 +21,8 @@ DECLARE_DELEGATE_TwoParams(FInstallBundleRemovedDelegate, TSharedRef<IInstallBun
 
 DECLARE_DELEGATE_OneParam(FInstallBundleReportingDelegate, const FInstallBundleReport&);
 
-DECLARE_DELEGATE_FourParams(FInstallBundleRequestStartedDelegate, const FString& /*BundleName*/, const FString& /*ManifestName*/, const FString& /*ManifestPath*/, const FString& /*ManifestVersion*/);
-DECLARE_DELEGATE_TwoParams(FInstallBundleRequestCompletedDelegate, const FString& /*BundleName*/, const FBundleRequestCompleteInfo& /*AnalyticInfo*/);
+DECLARE_DELEGATE_FourParams(FInstallBundleRequestStartedDelegate, const FName& /*BundleName*/, const FString& /*ManifestName*/, const FString& /*ManifestPath*/, const FString& /*ManifestVersion*/);
+DECLARE_DELEGATE_TwoParams(FInstallBundleRequestCompletedDelegate, const FName& /*BundleName*/, const FBundleRequestCompleteInfo& /*AnalyticInfo*/);
 
 DECLARE_DELEGATE_TwoParams(FInstallBundleSourceContentPatchResultDelegate, TSharedRef<IInstallBundleSource> /*Source*/, bool /*bContentPatchRequired*/);
 
