Add ChunkDownloadCompleteInfos into InstallBundleManager.BundleRequestCompleteDetailed event

#rnx

diff --git a/Engine/Plugins/Experimental/DefaultInstallBundleManager/Source/Private/DefaultInstallBundleManager.cpp b/Engine/Plugins/Experimental/DefaultInstallBundleManager/Source/Private/DefaultInstallBundleManager.cpp
--- a/Engine/Plugins/Experimental/DefaultInstallBundleManager/Source/Private/DefaultInstallBundleManager.cpp
+++ b/Engine/Plugins/Experimental/DefaultInstallBundleManager/Source/Private/DefaultInstallBundleManager.cpp
@@ -214,6 +214,8 @@ FDefaultInstallBundleManager::~FDefaultInstallBundleManager()
 	ReportingHandle.Reset();
 
 	InstallBundleUtil::CleanupInstallBundleAsyncIOTasks(AsyncIOTasks);
+	
+	IInstallBundleManager::InstallBundleChunkDownloadMetricsDelegate.RemoveAll(this);
 }
 
 void FDefaultInstallBundleManager::Initialize()
@@ -237,6 +239,10 @@ void FDefaultInstallBundleManager::Initialize()
 		LOG_INSTALL_BUNDLE_MAN(Warning, TEXT("Unrecoverable Initialization Failure - %s"), LexToString(InitResult));
 		bUnrecoverableInitError = true;
 	}
+	else
+	{
+		IInstallBundleManager::InstallBundleChunkDownloadMetricsDelegate.AddSP(this, &FDefaultInstallBundleManager::OnChunkDownloadMetrics);
+	}
 }
 
 bool FDefaultInstallBundleManager::Tick(float dt)
@@ -2300,8 +2306,10 @@ void FDefaultInstallBundleManager::TickContentRequests()
 							Request->bContentWasInstalled,
 							LexToString(Request->Result),
 							StatsMap->FindStatsChecked(Request->GetStatsKey()),
-							*Data);
+							*Data,
+							ChunkDownloadInfoMap.Find(Request->BundleName));
 						BundleRequestCompleteInfoMap.Remove(Request->BundleName);
+						ChunkDownloadInfoMap.Remove(Request->BundleName);
 					}
 					else
 					{
@@ -4740,6 +4748,18 @@ EInstallBundleManagerInitResult FDefaultInstallBundleManager::Init_TryCreateBund
 	return EInstallBundleManagerInitResult::OK;
 }
 
+void FDefaultInstallBundleManager::OnChunkDownloadMetrics(FInstallBundleChunkDownloadInfo ChunkDownloadInfo)
+{
+	if (auto ChunkDownloadInfos = ChunkDownloadInfoMap.Find(ChunkDownloadInfo.BundleName))
+	{
+		ChunkDownloadInfos->Add(ChunkDownloadInfo);
+	}
+	else
+	{
+		ChunkDownloadInfoMap.Add(ChunkDownloadInfo.BundleName, TArray<FInstallBundleChunkDownloadInfo>{ChunkDownloadInfo});
+	}
+}
+
 FInstallBundleSourceType FDefaultInstallBundleManager::FindFallbackSource(FInstallBundleSourceType SourceType)
 {
 	// See if we can find a fallback
diff --git a/Engine/Plugins/Experimental/DefaultInstallBundleManager/Source/Private/InstallBundleManagerUtil.cpp b/Engine/Plugins/Experimental/DefaultInstallBundleManager/Source/Private/InstallBundleManagerUtil.cpp
--- a/Engine/Plugins/Experimental/DefaultInstallBundleManager/Source/Private/InstallBundleManagerUtil.cpp
+++ b/Engine/Plugins/Experimental/DefaultInstallBundleManager/Source/Private/InstallBundleManagerUtil.cpp
@@ -789,7 +789,8 @@ namespace InstallBundleManagerAnalytics
 		bool bDidInstall, 
 		const FString& Result, 
 		const InstallBundleUtil::FContentRequestStats& TimingStats,
-		const FBundleRequestCompleteInfo& RequestCompleteInfo)
+		const FBundleRequestCompleteInfo& RequestCompleteInfo,
+		const TArray<FInstallBundleChunkDownloadInfo>* ChunkDownloadInfos)
 	{
 		if (AnalyticsProvider == nullptr || InstallBundleUtil::FInstallBundleSuppressAnalytics::IsEnabled())
 		{
@@ -817,6 +818,29 @@ namespace InstallBundleManagerAnalytics
 		Attributes.Add(FAnalyticsEventAttribute(TEXT("TotalDownloadedBytes"), RequestCompleteInfo.TotalDownloadedBytes));
 		Attributes.Add(FAnalyticsEventAttribute(TEXT("EstimatedFullDownloadBytes"), RequestCompleteInfo.EstimatedFullDownloadBytes));
 		Attributes.Add(FAnalyticsEventAttribute(TEXT("BPTErrorCode"), RequestCompleteInfo.BPTErrorCode));
+		
+		FString ChunckDonwloadInfoString;
+		FInstallBundleManagerUtilJsonWriter JsonWriter(&ChunckDonwloadInfoString);
+		JsonWriter.WriteArrayStart();
+		if (ChunkDownloadInfos)
+		{
+			for (int i = 0; i < ChunkDownloadInfos->Num(); i++)
+			{
+				const FInstallBundleChunkDownloadInfo& ChunkInfo = (*ChunkDownloadInfos)[i];
+				JsonWriter.WriteObjectStart();
+
+				JsonWriter.WriteValue(TEXT("ChunkName"), ChunkInfo.ChunkName);
+				JsonWriter.WriteValue(TEXT("ChunkFileSize"), ChunkInfo.ChunkFileSize);
+				JsonWriter.WriteValue(TEXT("ChunkDownloadDuration"), ChunkInfo.ChunkDownloadDuration);
+				JsonWriter.WriteValue(TEXT("FetchStartTimeUTC"), ChunkInfo.FetchStartTimeUTC.ToUnixTimestampDecimal());
+				JsonWriter.WriteValue(TEXT("FetchEndTimeUTC"), ChunkInfo.FetchEndTimeUTC.ToUnixTimestampDecimal());
+
+				JsonWriter.WriteObjectEnd();
+			}
+		}
+		JsonWriter.WriteArrayEnd();
+		JsonWriter.Close();
+		Attributes.Add(FAnalyticsEventAttribute(TEXT("ChunkDownloadCompleteInfos"), FJsonFragment(MoveTemp(ChunckDonwloadInfoString))));
 
 		AnalyticsProvider->RecordEvent(TEXT("InstallBundleManager.BundleRequestCompleteDetailed"), MoveTemp(Attributes));
 	}
diff --git a/Engine/Plugins/Experimental/DefaultInstallBundleManager/Source/Public/DefaultInstallBundleManager.h b/Engine/Plugins/Experimental/DefaultInstallBundleManager/Source/Public/DefaultInstallBundleManager.h
--- a/Engine/Plugins/Experimental/DefaultInstallBundleManager/Source/Public/DefaultInstallBundleManager.h
+++ b/Engine/Plugins/Experimental/DefaultInstallBundleManager/Source/Public/DefaultInstallBundleManager.h
@@ -635,6 +635,8 @@ protected:
 
 	UE_API bool CancelUpdateContentInternal(TArrayView<const FName> BundleNames);
 	UE_API bool CancelReleaseContentInternal(TArrayView<const FName> BundleNames);
+	
+	void OnChunkDownloadMetrics(FInstallBundleChunkDownloadInfo ChunkDownloadInfo);
 
 	// IInstallBundleManager interface
 public:
@@ -832,6 +834,7 @@ protected:
 	TSharedPtr<IAnalyticsProviderET> AnalyticsProvider;
 	TSharedRef<InstallBundleUtil::FContentRequestStatsMap> StatsMap;
 	TMap<FName, FBundleRequestCompleteInfo> BundleRequestCompleteInfoMap;	// BundleName as key
+	TMap<FName, TArray<FInstallBundleChunkDownloadInfo> > ChunkDownloadInfoMap;	// BundleName as key
 };
 
 
diff --git a/Engine/Plugins/Experimental/DefaultInstallBundleManager/Source/Public/InstallBundleManagerUtil.h b/Engine/Plugins/Experimental/DefaultInstallBundleManager/Source/Public/InstallBundleManagerUtil.h
--- a/Engine/Plugins/Experimental/DefaultInstallBundleManager/Source/Public/InstallBundleManagerUtil.h
+++ b/Engine/Plugins/Experimental/DefaultInstallBundleManager/Source/Public/InstallBundleManagerUtil.h
@@ -15,6 +15,7 @@ class IInstallBundleSource;
 struct FAnalyticsEventAttribute;
 class IAnalyticsProviderET;
 struct FInstallBundleCacheStats;
+struct FInstallBundleChunkDownloadInfo;
 
 DECLARE_LOG_CATEGORY_EXTERN(LogDefaultInstallBundleManager, Display, All);
 
@@ -368,13 +369,15 @@ namespace InstallBundleManagerAnalytics
 	 * @EventParam TotalDownloadedBytes (uint64) Total bytes downloaded
 	 * @EventParam EstimatedFullDownloadBytes (uint64) Estimated bytes that needed download
 	 * @EventParam BPTErrorCode (string) Any BPT Error
-	 */
+	 * @EventParam ChunkDownloadCompleteInfos (string) Detail metrics of each chunk downloaded, if any
+	*/
 	DEFAULTINSTALLBUNDLEMANAGER_API void FireEvent_BundleRequestCompleteDetailed(IAnalyticsProviderET* AnalyticsProvider,
 		const FString& BundleName,
 		bool bDidInstall,
 		const FString& Result,
 		const InstallBundleUtil::FContentRequestStats& TimingStats,
-		const FBundleRequestCompleteInfo& RequestCompleteInfo);
+		const FBundleRequestCompleteInfo& RequestCompleteInfo,
+		const TArray<FInstallBundleChunkDownloadInfo>* ChunkDownloadInfos);
 
 
 	/**
diff --git a/Engine/Source/Runtime/InstallBundleManager/Public/InstallBundleManagerInterface.h b/Engine/Source/Runtime/InstallBundleManager/Public/InstallBundleManagerInterface.h
--- a/Engine/Source/Runtime/InstallBundleManager/Public/InstallBundleManagerInterface.h
+++ b/Engine/Source/Runtime/InstallBundleManager/Public/InstallBundleManagerInterface.h
@@ -93,6 +93,7 @@ enum class EInstallBundleManagerInitErrorHandlerResult
 
 struct FInstallBundleChunkDownloadInfo
 {
+	FName BundleName;
 	FString ChunkName;
 	int32 ChunkFileSize;
 	float ChunkDownloadDuration;
