UE: Give InstallBundleSourceBulk subclasses the opportunity to modify the content paths for a RequestUpdateContent call (optional)

diff --git a/Engine/Plugins/Experimental/DefaultInstallBundleManager/Source/Private/InstallBundleSourceBulk.cpp b/Engine/Plugins/Experimental/DefaultInstallBundleManager/Source/Private/InstallBundleSourceBulk.cpp
--- a/Engine/Plugins/Experimental/DefaultInstallBundleManager/Source/Private/InstallBundleSourceBulk.cpp
+++ b/Engine/Plugins/Experimental/DefaultInstallBundleManager/Source/Private/InstallBundleSourceBulk.cpp
@@ -622,6 +622,8 @@ void FInstallBundleSourceBulk::RequestUpdateContent(FRequestUpdateContentBundleC
 		ResultInfo.ContentPaths = *BundleFileList;
 	}
 
+	FilterContentBundlePaths(Context, ResultInfo.ContentPaths);
+
 	for (const FString& Path : ResultInfo.ContentPaths)
 	{
 		if (Path.EndsWith(TEXTVIEW(".uondemandtoc")))
diff --git a/Engine/Plugins/Experimental/DefaultInstallBundleManager/Source/Public/InstallBundleSourceBulk.h b/Engine/Plugins/Experimental/DefaultInstallBundleManager/Source/Public/InstallBundleSourceBulk.h
--- a/Engine/Plugins/Experimental/DefaultInstallBundleManager/Source/Public/InstallBundleSourceBulk.h
+++ b/Engine/Plugins/Experimental/DefaultInstallBundleManager/Source/Public/InstallBundleSourceBulk.h
@@ -67,6 +67,9 @@ protected:
 	UE_API virtual void GetOnDemandHostGroup(UE::IoStore::FOnDemandHostGroup& OutHostGroup);
 	UE_API virtual FString GetOnDemandTocRelativeURL();
 
+	// Optionally used to modify what ContentPaths are returned in the base implementation of RequestUpdateContent
+	UE_API virtual void FilterContentBundlePaths(const FRequestUpdateContentBundleContext& BundleContext, TArray<FString>& InOutContentPaths) {}
+
 protected:
 	FTSTicker::FDelegateHandle TickHandle;
 
