Fixup API macros

#rnx

diff --git a/Engine/Plugins/Experimental/DefaultInstallBundleManager/Source/Public/InstallBundleSourceBulk.h b/Engine/Plugins/Experimental/DefaultInstallBundleManager/Source/Public/InstallBundleSourceBulk.h
--- a/Engine/Plugins/Experimental/DefaultInstallBundleManager/Source/Public/InstallBundleSourceBulk.h
+++ b/Engine/Plugins/Experimental/DefaultInstallBundleManager/Source/Public/InstallBundleSourceBulk.h
@@ -68,7 +68,7 @@ protected:
 	UE_API virtual FString GetOnDemandTocRelativeURL();
 
 	// Optionally used to modify what ContentPaths are returned in the base implementation of RequestUpdateContent
-	UE_API virtual void FilterContentBundlePaths(const FRequestUpdateContentBundleContext& BundleContext, TArray<FString>& InOutContentPaths) {}
+	virtual void FilterContentBundlePaths(const FRequestUpdateContentBundleContext& BundleContext, TArray<FString>& InOutContentPaths) {}
 
 protected:
 	FTSTicker::FDelegateHandle TickHandle;
