Change QueryPersistentBundleInfo to match the behavior of the bespoke platform bundle sources when encountering an invalid named chunk.
- typically only applies to a malformed package

#rnx

diff --git a/Engine/Plugins/Experimental/DefaultInstallBundleManager/Source/Private/InstallBundleSourcePlatformChunkInstall.cpp b/Engine/Plugins/Experimental/DefaultInstallBundleManager/Source/Private/InstallBundleSourcePlatformChunkInstall.cpp
--- a/Engine/Plugins/Experimental/DefaultInstallBundleManager/Source/Private/InstallBundleSourcePlatformChunkInstall.cpp
+++ b/Engine/Plugins/Experimental/DefaultInstallBundleManager/Source/Private/InstallBundleSourcePlatformChunkInstall.cpp
@@ -388,7 +388,7 @@ bool FInstallBundleSourcePlatformChunkInstall::QueryPersistentBundleInfo(FInstal
 		return true;
 	}
 
-	return false;
+	return true;
 }
 
 
