Implementing the Launcher Install Bundle Source to act as an interface for installing content bundles using the Epic Online Services Helper (EOSH). Bundles can be assigned to a seperate artifact from the core game content in the artifact service and requested through this install bundle source.

diff --git a/Engine/Plugins/Experimental/DefaultInstallBundleManager/Source/Private/DefaultInstallBundleManager.cpp b/Engine/Plugins/Experimental/DefaultInstallBundleManager/Source/Private/DefaultInstallBundleManager.cpp
--- a/Engine/Plugins/Experimental/DefaultInstallBundleManager/Source/Private/DefaultInstallBundleManager.cpp
+++ b/Engine/Plugins/Experimental/DefaultInstallBundleManager/Source/Private/DefaultInstallBundleManager.cpp
@@ -51,6 +51,41 @@ static FAutoConsoleCommand RequestBundleCommand(
 }),
 ECVF_Cheat);
 
+static FAutoConsoleCommand RequestReleaseBundleCommand(
+	TEXT("RequestReleaseBundle"),
+	TEXT("Request a bundle to be uninstalled"),
+	FConsoleCommandWithArgsDelegate::CreateLambda([](const TArray<FString>& Args)
+{
+	TSharedPtr<IInstallBundleManager> InstallBundleManager = IInstallBundleManager::GetPlatformInstallBundleManager();
+	if (InstallBundleManager == nullptr)
+	{
+		LOG_INSTALL_BUNDLE_MAN(Error, TEXT("RequestReleaseBundle: InstallBundleManager not found!"));
+		return;
+	}
+	if (InstallBundleManager->GetInitState() != EInstallBundleManagerInitState::Succeeded)
+	{
+		LOG_INSTALL_BUNDLE_MAN(Error, TEXT("RequestReleaseBundle: InstallBundleManager not initialized!"));
+		return;
+	}
+	if (Args.Num() != 1)
+	{
+		LOG_INSTALL_BUNDLE_MAN(Display, TEXT("RequestReleaseBundle: Usage:\n")
+							   TEXT("RequestReleaseBundle <BundleName>"));
+		return;
+	}
+	TValueOrError<FInstallBundleCombinedInstallState, EInstallBundleResult> InstallState = InstallBundleManager->GetInstallStateSynchronous(FName(*Args[0]), false);
+	if (InstallState.HasValue())
+	{
+		LOG_INSTALL_BUNDLE_MAN(Warning, TEXT("Requesting to uninstall bundle %s"), *Args[0]);
+		InstallBundleManager->RequestReleaseContent(FName(*Args[0]), EInstallBundleReleaseRequestFlags::ExplicitRemoveList | EInstallBundleReleaseRequestFlags::RemoveFilesIfPossible);
+	}
+	else
+	{
+		LOG_INSTALL_BUNDLE_MAN(Error, TEXT("RequestBundle: InstallBundleManager not found!"));
+	}
+}),
+ECVF_Cheat);
+
 static FAutoConsoleCommand SetInstallBundleErrorSimulationCommands(
 	TEXT("SetInstallBundleErrorSimulationCommands"),
 	TEXT("Set flags for simulating install bundle errors"),
