Make sure we can do a partial install using DI and then go back to ChunkDBs

#rnx

diff --git a/Engine/Plugins/Experimental/DefaultInstallBundleManager/Source/Private/DefaultInstallBundleManager.cpp b/Engine/Plugins/Experimental/DefaultInstallBundleManager/Source/Private/DefaultInstallBundleManager.cpp
--- a/Engine/Plugins/Experimental/DefaultInstallBundleManager/Source/Private/DefaultInstallBundleManager.cpp
+++ b/Engine/Plugins/Experimental/DefaultInstallBundleManager/Source/Private/DefaultInstallBundleManager.cpp
@@ -4726,7 +4726,7 @@ EInstallBundleManagerInitResult FDefaultInstallBundleManager::Init_TryCreateBund
 								{
 									const uint64 TotalPersistent = PersistentReport.TotalDownloadSize();
 									const uint64 TotalNew = BundleReport.TotalDownloadSize();
-									if (ensureMsgf(TotalNew <= TotalPersistent, TEXT("Bundle %s has TotalPersistent=%" UINT64_FMT ", and TotalNew=%" UINT64_FMT ", which is not expected"), *BundleReport.BundleName.ToString(), TotalPersistent, TotalNew))
+									if (TotalNew <= TotalPersistent)
 									{
 										if (const uint64 Discrepancy = TotalPersistent - TotalNew)
 										{
