Allow checking if an install bundle is mounted as well as up to date.  This better supports flows where installation is seperate from mounting.

diff --git a/Engine/Plugins/Experimental/DefaultInstallBundleManager/Source/Private/DefaultInstallBundleManager.cpp b/Engine/Plugins/Experimental/DefaultInstallBundleManager/Source/Private/DefaultInstallBundleManager.cpp
--- a/Engine/Plugins/Experimental/DefaultInstallBundleManager/Source/Private/DefaultInstallBundleManager.cpp
+++ b/Engine/Plugins/Experimental/DefaultInstallBundleManager/Source/Private/DefaultInstallBundleManager.cpp
@@ -674,8 +674,11 @@ FInstallBundleCombinedInstallState FDefaultInstallBundleManager::GetInstallState
 			InstallState = EInstallBundleInstallState::NeedsUpdate;
 			break;
 		case EBundleState::NeedsMount:
+			InstallState = EInstallBundleInstallState::UpToDate;
+			break;
 		case EBundleState::Mounted:
 			InstallState = EInstallBundleInstallState::UpToDate;
+			RetVal.MountedBundles.Add(BundleName);
 			break;
 		default:
 			ensureMsgf(false, TEXT("Unknown EBundleState"));
diff --git a/Engine/Source/Runtime/InstallBundleManager/Public/InstallBundleTypes.h b/Engine/Source/Runtime/InstallBundleManager/Public/InstallBundleTypes.h
--- a/Engine/Source/Runtime/InstallBundleManager/Public/InstallBundleTypes.h
+++ b/Engine/Source/Runtime/InstallBundleManager/Public/InstallBundleTypes.h
@@ -114,6 +114,7 @@ INSTALLBUNDLEMANAGER_API const TCHAR* LexToString(EInstallBundleInstallState Sta
 struct FInstallBundleCombinedInstallState
 {
 	TMap<FName, EInstallBundleInstallState> IndividualBundleStates;
+	TSet<FName> MountedBundles;
 	TSet<FName> BundlesWithIoStoreOnDemand;
 
 	INSTALLBUNDLEMANAGER_API bool GetAllBundlesHaveState(EInstallBundleInstallState State, TArrayView<const FName> ExcludedBundles = TArrayView<const FName>()) const;
