Add analytics fields to track download type

#rnx
#Synced-CL 48150757

diff --git a/Engine/Plugins/Experimental/DefaultInstallBundleManager/Source/Private/InstallBundleManagerUtil.cpp b/Engine/Plugins/Experimental/DefaultInstallBundleManager/Source/Private/InstallBundleManagerUtil.cpp
--- a/Engine/Plugins/Experimental/DefaultInstallBundleManager/Source/Private/InstallBundleManagerUtil.cpp
+++ b/Engine/Plugins/Experimental/DefaultInstallBundleManager/Source/Private/InstallBundleManagerUtil.cpp
@@ -818,6 +818,7 @@ namespace InstallBundleManagerAnalytics
 		Attributes.Add(FAnalyticsEventAttribute(TEXT("TotalDownloadedBytes"), RequestCompleteInfo.TotalDownloadedBytes));
 		Attributes.Add(FAnalyticsEventAttribute(TEXT("EstimatedFullDownloadBytes"), RequestCompleteInfo.EstimatedFullDownloadBytes));
 		Attributes.Add(FAnalyticsEventAttribute(TEXT("BPTErrorCode"), RequestCompleteInfo.BPTErrorCode));
+		Attributes.Add(FAnalyticsEventAttribute(TEXT("DownloadType"), RequestCompleteInfo.DownloadType));
 		
 		FString ChunckDonwloadInfoString;
 		FInstallBundleManagerUtilJsonWriter JsonWriter(&ChunckDonwloadInfoString);
diff --git a/Engine/Plugins/Experimental/DefaultInstallBundleManager/Source/Public/InstallBundleManagerUtil.h b/Engine/Plugins/Experimental/DefaultInstallBundleManager/Source/Public/InstallBundleManagerUtil.h
--- a/Engine/Plugins/Experimental/DefaultInstallBundleManager/Source/Public/InstallBundleManagerUtil.h
+++ b/Engine/Plugins/Experimental/DefaultInstallBundleManager/Source/Public/InstallBundleManagerUtil.h
@@ -369,6 +369,7 @@ namespace InstallBundleManagerAnalytics
 	 * @EventParam TotalDownloadedBytes (uint64) Total bytes downloaded
 	 * @EventParam EstimatedFullDownloadBytes (uint64) Estimated bytes that needed download
 	 * @EventParam BPTErrorCode (string) Any BPT Error
+	 * @EventParam DownloadType (string) Type of download
 	 * @EventParam ChunkDownloadCompleteInfos (string) Detail metrics of each chunk downloaded, if any
 	*/
 	DEFAULTINSTALLBUNDLEMANAGER_API void FireEvent_BundleRequestCompleteDetailed(IAnalyticsProviderET* AnalyticsProvider,
diff --git a/Engine/Source/Runtime/InstallBundleManager/Public/InstallBundleTypes.h b/Engine/Source/Runtime/InstallBundleManager/Public/InstallBundleTypes.h
--- a/Engine/Source/Runtime/InstallBundleManager/Public/InstallBundleTypes.h
+++ b/Engine/Source/Runtime/InstallBundleManager/Public/InstallBundleTypes.h
@@ -515,4 +515,5 @@ struct FBundleRequestCompleteInfo
 	uint64 EstimatedFullDownloadBytes;
 	FString Result;
 	FString BPTErrorCode;
+	FString DownloadType;
 };
