[TRACE] Improve reporting of BundleManager, use regions instead of bookmarks to show a zone where the bundle is mounted

#rnx
#Synced-CL Unknown

diff --git a/Engine/Plugins/Experimental/DefaultInstallBundleManager/Source/Private/DefaultInstallBundleManager.cpp b/Engine/Plugins/Experimental/DefaultInstallBundleManager/Source/Private/DefaultInstallBundleManager.cpp
--- a/Engine/Plugins/Experimental/DefaultInstallBundleManager/Source/Private/DefaultInstallBundleManager.cpp
+++ b/Engine/Plugins/Experimental/DefaultInstallBundleManager/Source/Private/DefaultInstallBundleManager.cpp
@@ -20,6 +20,7 @@
 #include "Algo/AnyOf.h"
 
 #include "IAnalyticsProviderET.h"
+#include "UObject/NameTypes.h"
 
 #if WITH_EDITOR
 #include "Misc/KeyChainUtilities.h"
@@ -1645,7 +1646,7 @@ void FDefaultInstallBundleManager::MountPaks(FContentRequestRef Request)
 	}
 
 	SCOPED_BOOT_TIMING("FDefaultInstallBundleManager::MountPaks");
-	TRACE_BOOKMARK(TEXT("Start Mount Bundle %s"), *Request->BundleName.ToString());
+	TRACE_BEGIN_REGION(*WriteToString<256>(TEXT("Mount "), FNameBuilder(Request->BundleName)), TEXT("InstallBundles"));
 
 	FBundleInfo& BundleInfo = BundleInfoMap.FindChecked(Request->BundleName);
 
@@ -1697,7 +1698,7 @@ void FDefaultInstallBundleManager::MountPaks(FContentRequestRef Request)
 		SetBundleStatus(BundleInfoLambda, EBundleState::Mounted);
 
 		StatsEnd(Request->GetStatsKey(), EContentRequestState::Mounting);
-		TRACE_BOOKMARK(TEXT("Finished Mount Bundle %s"), *Request->BundleName.ToString());
+		TRACE_END_REGION(*WriteToString<256>(TEXT("Mount "), FNameBuilder(Request->BundleName)));
 		Request->StepResult = EContentRequestStepResult::Done;
 	};
 
@@ -1758,7 +1759,7 @@ void FDefaultInstallBundleManager::MountPaks(FContentRequestRef Request)
 		{
 			LOG_INSTALL_BUNDLE_MAN_OVERRIDE(Request->LogVerbosityOverride, Display, TEXT("Skipping Mounting Paks for Request %s"), *Request->BundleName.ToString());
 
-			TRACE_BOOKMARK(TEXT("Finished Mount Bundle %s"), *Request->BundleName.ToString());
+			TRACE_END_REGION(*WriteToString<256>(TEXT("Mount "), FNameBuilder(Request->BundleName)));
 			Request->StepResult = EContentRequestStepResult::Done;
 			return;
 		}
@@ -1817,7 +1818,7 @@ void FDefaultInstallBundleManager::MountPaks(FContentRequestRef Request)
 						ErrorCode.ReplaceCharInline(TEXT(' '), TEXT('_'), ESearchCase::CaseSensitive);
 
 						PinnedThis->StatsEnd(Request->GetStatsKey(), EContentRequestState::Mounting);
-						TRACE_BOOKMARK(TEXT("Finished Mount Bundle %s"), *Request->BundleName.ToString());
+						TRACE_END_REGION(*WriteToString<256>(TEXT("Mount "), FNameBuilder(Request->BundleName)));
 
 						Request->Result = EInstallBundleResult::InstallError;
 						Request->OptionalErrorCode = MoveTemp(ErrorCode);
@@ -1835,7 +1836,7 @@ void FDefaultInstallBundleManager::MountPaks(FContentRequestRef Request)
 					LOG_INSTALL_BUNDLE_MAN_OVERRIDE(Request->LogVerbosityOverride, Display, TEXT("Skipping Mounting Paks for Request %s"), *Request->BundleName.ToString());
 
 					PinnedThis->StatsEnd(Request->GetStatsKey(), EContentRequestState::Mounting);
-					TRACE_BOOKMARK(TEXT("Finished Mount Bundle %s"), *Request->BundleName.ToString());
+					TRACE_END_REGION(*WriteToString<256>(TEXT("Mount "), FNameBuilder(Request->BundleName)));
 					Request->StepResult = EContentRequestStepResult::Done;
 					return;
 				}
