Add CSV metadata BundleManager_HasDownloadedContent so we can track whether content has been downloaded since boot
#Synced-CL Unknown

diff --git a/Engine/Plugins/Experimental/DefaultInstallBundleManager/Source/Private/DefaultInstallBundleManager.cpp b/Engine/Plugins/Experimental/DefaultInstallBundleManager/Source/Private/DefaultInstallBundleManager.cpp
--- a/Engine/Plugins/Experimental/DefaultInstallBundleManager/Source/Private/DefaultInstallBundleManager.cpp
+++ b/Engine/Plugins/Experimental/DefaultInstallBundleManager/Source/Private/DefaultInstallBundleManager.cpp
@@ -14,6 +14,7 @@
 #include "Misc/EnumClassFlags.h"
 #include "MountPaksExArgs.h"
 #include "PatchCheck.h"
+#include "ProfilingDebugging/CsvProfiler.h"
 #include "Templates/Greater.h"
 #include "ShaderPipelineCache.h"
 #include "Algo/Find.h"
@@ -241,6 +242,8 @@ FDefaultInstallBundleManager::FDefaultInstallBundleManager(FInstallBundleSourceF
 
 	SetErrorSimulationCommands(FCommandLine::Get());
 	SetCommandLineOverrides(FCommandLine::Get());
+
+	CSV_METADATA(TEXT("BundleManager_HasDownloadedContent"), TEXT("0"));
 }
 
 FDefaultInstallBundleManager::~FDefaultInstallBundleManager()
@@ -2666,6 +2669,7 @@ bool FDefaultInstallBundleManager::TickReporting(float dt)
 		// Some platforms do a miniature 30-ish byte update after they've been patched from the outside (e.g. EGL), we don't want to count those as "has updated anything".
 		if (!bHasEverUpdatedContent && PersistentState.TotalDownloadSize() > 1024)
 		{
+			CSV_METADATA(TEXT("BundleManager_HasDownloadedContent"), TEXT("1"));
 			bHasEverUpdatedContent = true;
 		}
 	}
