Clean-up of on-demand mount arguments.
Since most of the on-demand container settings are data driven from the TOC header, this change
removes most of the additional mount argument except the file path to the TOC.

- removed host group argument since the host group name is serialized in the TOC header
- removed the use of ForceNonSecureHTTP since we now use HTTP2 and TLS is mandatory
- removed setting relative TOC path except for UFEN since the path is serialized in the TOC header
- removed specifying InstallOnDemand since the flags are serialized in the TOC header
- removed GetOnDemandHostGroup and GetOnDemandTocRelativeURL from FInstallBundleSourceBulk since these the methods are no longer used

#rnx

#Synced-CL 49721519

diff --git a/Engine/Plugins/Experimental/DefaultInstallBundleManager/Source/Private/InstallBundleSourceBulk.cpp b/Engine/Plugins/Experimental/DefaultInstallBundleManager/Source/Private/InstallBundleSourceBulk.cpp
--- a/Engine/Plugins/Experimental/DefaultInstallBundleManager/Source/Private/InstallBundleSourceBulk.cpp
+++ b/Engine/Plugins/Experimental/DefaultInstallBundleManager/Source/Private/InstallBundleSourceBulk.cpp
@@ -533,17 +533,6 @@ void FInstallBundleSourceBulk::SerializeBulkBuildBundleMetadata(const TMap<FName
 	}
 }
 
-void FInstallBundleSourceBulk::GetOnDemandHostGroup(UE::IoStore::FOnDemandHostGroup& OutHostGroup)
-{
-	ensureAlwaysMsgf(false, TEXT("FInstallBundleSourceBulk::GetOnDemandHostGroup not implemented!"));
-}
-
-FString FInstallBundleSourceBulk::GetOnDemandTocRelativeURL()
-{
-	ensureAlwaysMsgf(false, TEXT("FInstallBundleSourceBulk::GetOnDemandTocRelativeURL not implemented!"));
-	return {};
-}
-
 EInstallBundleInstallState FInstallBundleSourceBulk::GetBundleInstallState(FName BundleName)
 {
 	return EInstallBundleInstallState::UpToDate;
@@ -678,12 +667,7 @@ void FInstallBundleSourceBulk::RequestUpdateContent(FRequestUpdateContentBundleC
 		{
 			TUniquePtr<UE::IoStore::FOnDemandMountArgs> MountArgs = MakeUnique<UE::IoStore::FOnDemandMountArgs>();
 			MountArgs->MountId = Context.BundleName.ToString();
-			GetOnDemandHostGroup(MountArgs->HostGroup);
-			MountArgs->TocRelativeUrl = GetOnDemandTocRelativeURL();
 			MountArgs->FilePath = Path;
-			MountArgs->Options = UE::IoStore::EOnDemandMountOptions::InstallOnDemand;
-			// IAX content on the BPS CDN's does not require https
-			MountArgs->Options |= UE::IoStore::EOnDemandMountOptions::ForceNonSecureHttp;
 
 			if (OnDemandMountOptions.bWithSoftReferences)
 			{
diff --git a/Engine/Plugins/Experimental/DefaultInstallBundleManager/Source/Public/InstallBundleSourceBulk.h b/Engine/Plugins/Experimental/DefaultInstallBundleManager/Source/Public/InstallBundleSourceBulk.h
--- a/Engine/Plugins/Experimental/DefaultInstallBundleManager/Source/Public/InstallBundleSourceBulk.h
+++ b/Engine/Plugins/Experimental/DefaultInstallBundleManager/Source/Public/InstallBundleSourceBulk.h
@@ -64,8 +64,6 @@ public:
 
 	// Support for on demand tocs, must be implemented by the game
 protected:
-	UE_API virtual void GetOnDemandHostGroup(UE::IoStore::FOnDemandHostGroup& OutHostGroup);
-	UE_API virtual FString GetOnDemandTocRelativeURL();
 
 	// Optionally used to modify what ContentPaths are returned in the base implementation of RequestUpdateContent
 	virtual void FilterContentBundlePaths(const FRequestUpdateContentBundleContext& BundleContext, TArray<FString>& InOutContentPaths) {}
