Make Install Bundle Source Platform Chunk Install respect mounting iostore soft refs and return OnDemand mount options.
- will follow this up with a refactor to remove needing to return OnDemand mount options.

#Synced-CL 50152832

diff --git a/Engine/Plugins/Experimental/DefaultInstallBundleManager/Source/Private/InstallBundleSourcePlatformChunkInstall.cpp b/Engine/Plugins/Experimental/DefaultInstallBundleManager/Source/Private/InstallBundleSourcePlatformChunkInstall.cpp
--- a/Engine/Plugins/Experimental/DefaultInstallBundleManager/Source/Private/InstallBundleSourcePlatformChunkInstall.cpp
+++ b/Engine/Plugins/Experimental/DefaultInstallBundleManager/Source/Private/InstallBundleSourcePlatformChunkInstall.cpp
@@ -12,6 +12,7 @@
 #include "Algo/AnyOf.h"
 #include "Algo/Find.h"
 #include "Containers/Ticker.h"
+#include "IO/IoStoreOnDemand.h"
 #include "IPlatformFilePak.h"
 #include "Stats/Stats.h"
 
@@ -532,6 +533,39 @@ void FInstallBundleSourcePlatformChunkInstall::RequestUpdateContent(FRequestUpda
 		ResultInfo.Result = EInstallBundleResult::OK;
 		ResultInfo.ContentPaths = BundleInfo->FilePaths;
 		ResultInfo.bContentWasInstalled = false;
+
+		InstallBundleUtil::FConfigMountOptions ConfigMountOptions;
+		InstallBundleUtil::GetMountOptionsFromConfig(FNameBuilder(Context.BundleName), ConfigMountOptions);
+
+		if (ConfigMountOptions.bWithSoftReferences)
+		{
+			ResultInfo.MountOptions.MountFlags |= FPakMountOptions::EMountFlags::WithSoftReferences;
+		}
+
+		// TODO: this no longer needs to happen in the bundle source
+		// Move this into IBM
+		for (const FString& Path : ResultInfo.ContentPaths)
+		{
+			if (Path.EndsWith(TEXTVIEW(".uondemandtoc")))
+			{
+				TUniquePtr<UE::IoStore::FOnDemandMountArgs> MountArgs = MakeUnique<UE::IoStore::FOnDemandMountArgs>();
+				MountArgs->MountId = Context.BundleName.ToString();
+				//GetOnDemandHostGroup(MountArgs->HostGroup);
+				//MountArgs->TocRelativeUrl = GetOnDemandTocRelativeURL();
+				MountArgs->FilePath = Path;
+				MountArgs->Options = UE::IoStore::EOnDemandMountOptions::InstallOnDemand;
+				// IAX content on the BPS CDN's does not require https
+				MountArgs->Options |= UE::IoStore::EOnDemandMountOptions::ForceNonSecureHttp;
+
+				if (ConfigMountOptions.bWithSoftReferences)
+				{
+					MountArgs->Options |= UE::IoStore::EOnDemandMountOptions::WithSoftReferences;
+				}
+
+				ResultInfo.OnDemandMountArgs.Add(MoveTemp(MountArgs));
+			}
+		}
+
 		Context.CompleteCallback.ExecuteIfBound(AsShared(), MoveTemp(ResultInfo));
 	}
 	else
@@ -697,6 +731,38 @@ void FInstallBundleSourcePlatformChunkInstall::OnNamedChunkInstall( const FNamed
 				ResultInfo.bContentWasInstalled = true;
 			}
 
+			InstallBundleUtil::FConfigMountOptions ConfigMountOptions;
+			InstallBundleUtil::GetMountOptionsFromConfig(FNameBuilder(Request->BundleName), ConfigMountOptions);
+
+			if (ConfigMountOptions.bWithSoftReferences)
+			{
+				ResultInfo.MountOptions.MountFlags |= FPakMountOptions::EMountFlags::WithSoftReferences;
+			}
+
+			// TODO: this no longer needs to happen in the bundle source
+			// Move this into IBM
+			for (const FString& Path : ResultInfo.ContentPaths)
+			{
+				if (Path.EndsWith(TEXTVIEW(".uondemandtoc")))
+				{
+					TUniquePtr<UE::IoStore::FOnDemandMountArgs> MountArgs = MakeUnique<UE::IoStore::FOnDemandMountArgs>();
+					MountArgs->MountId = Request->BundleName.ToString();
+					//GetOnDemandHostGroup(MountArgs->HostGroup);
+					//MountArgs->TocRelativeUrl = GetOnDemandTocRelativeURL();
+					MountArgs->FilePath = Path;
+					MountArgs->Options = UE::IoStore::EOnDemandMountOptions::InstallOnDemand;
+					// IAX content on the BPS CDN's does not require https
+					MountArgs->Options |= UE::IoStore::EOnDemandMountOptions::ForceNonSecureHttp;
+
+					if (ConfigMountOptions.bWithSoftReferences)
+					{
+						MountArgs->Options |= UE::IoStore::EOnDemandMountOptions::WithSoftReferences;
+					}
+
+					ResultInfo.OnDemandMountArgs.Add(MoveTemp(MountArgs));
+				}
+			}
+
 			Request->CompleteCallback.ExecuteIfBound(AsShared(), MoveTemp(ResultInfo));
 		}
 
