Optionally parallize regex matching in InstallBundleSourceBulk.
Regex matching can take considerable time with many files and bundles.

#Synced-CL 49967858

#Synced-CL 49967858

#Synced-CL 49967858

diff --git a/Engine/Plugins/Experimental/DefaultInstallBundleManager/Source/Private/InstallBundleSourceBulk.cpp b/Engine/Plugins/Experimental/DefaultInstallBundleManager/Source/Private/InstallBundleSourceBulk.cpp
--- a/Engine/Plugins/Experimental/DefaultInstallBundleManager/Source/Private/InstallBundleSourceBulk.cpp
+++ b/Engine/Plugins/Experimental/DefaultInstallBundleManager/Source/Private/InstallBundleSourceBulk.cpp
@@ -11,10 +11,12 @@
 #include "InstallBundleManagerInterface.h"
 #include "InstallBundleManagerUtil.h"
 #include "Misc/CommandLine.h"
+#include "Misc/Parse.h"
 #include "Misc/FileHelper.h"
 #include "Misc/ConfigCacheIni.h"
 #include "Internationalization/Regex.h"
 #include "Misc/App.h"
+#include "Async/ParallelFor.h"
 
 #define LOG_SOURCE_BULK(Verbosity, Format, ...) LOG_INSTALL_BUNDLE_MAN(Verbosity, TEXT("InstallBundleSourceBulk: ") Format, ##__VA_ARGS__)
 
@@ -397,17 +399,68 @@ void FInstallBundleSourceBulk::AsyncInit_MakeBundlesForBulkBuild()
 			//If we had files remaining to manually parse, now sort the remaining files into bundles
 			if (!bDidLoadAllFilesFromMetadata)
 			{
-				for (const FString& Section : SectionNames)
+				TRACE_CPUPROFILER_EVENT_SCOPE(FInstallBundleSourceBulk_MatchFiles);
+				
+				struct FSectionParallelData
+				{
+					FString BundleName;
+					TBitArray<> MatchedFiles;
+					bool bAnyMatches;
+				};
+				TArray<FSectionParallelData> ParallelData;
+				ParallelData.AddDefaulted(SectionNames.Num());
+				
+				// If the file was matched by a bundle with higher priority we can skip it
+				// Used as an optimization - don't trust the value in here!
+				TArray<std::atomic<int32>> MatchedBy;
+				MatchedBy.AddDefaulted(FoundFiles->Num());
+				for (int32 I = 0; I < MatchedBy.Num(); ++I)
+				{
+					MatchedBy[I] = INT32_MAX;
+				}
+				
+				EParallelForFlags ParallelFlags = EParallelForFlags::Unbalanced;
+				
+				bool bParallelRegexMatching = true;
+				
+				// Search configs if we should MT
+				{
+					GConfig->GetBool(TEXT("InstallBundleSource.Bulk.MiscSettings"), TEXT("bParallelRegexMatching"), bParallelRegexMatching, GInstallBundleIni);
+	
+					if (const FString* ParallelMatchingVar = FPlatformMisc::GetConfigRulesVariable(TEXT("bParallelRegexMatching")))
+					{
+						bParallelRegexMatching = FCString::ToBool(**ParallelMatchingVar);
+					}
+					
+#if !UE_BUILD_SHIPPING
+					if (FParse::Param(FCommandLine::Get(), TEXT("NoParallelBundleRegexMatching")))
+					{
+						bParallelRegexMatching = false;
+					}
+#endif // !UE_BUILD_SHIPPING
+				}
+				
+				if (!bParallelRegexMatching)
+				{
+					ParallelFlags |= EParallelForFlags::ForceSingleThread;
+				}
+				
+				// Match found files for each section in parallel
+				ParallelFor(SectionNames.Num(), [&](int32 SectionIndex)
 				{
-					const FString BundleName = Section.RightChop(InstallBundleUtil::GetInstallBundleSectionPrefix().Len());
-					TArray<FString>& BundleFileList = BulkBuildBundles.FindOrAdd(FName(*BundleName));
+					const FString& Section = SectionNames[SectionIndex];
+					FString BundleName = Section.RightChop(InstallBundleUtil::GetInstallBundleSectionPrefix().Len());
+					FSectionParallelData& Data = ParallelData[SectionIndex];
+					Data.BundleName = BundleName;
+					Data.MatchedFiles.Init(false, FoundFiles->Num());
+					Data.bAnyMatches = false;
 
 					TArray<FString> StrSearchRegexPatterns;
 					InstallBundleConfig->GetArray(*Section, TEXT("FileRegex"), StrSearchRegexPatterns);
 
 					if (StrSearchRegexPatterns.IsEmpty())
 					{
-						continue;
+						return;
 					}
 
 					TArray<FString> BundleSubsets;
@@ -420,7 +473,7 @@ void FInstallBundleSourceBulk::AsyncInit_MakeBundlesForBulkBuild()
 						const FString SubsetSection = InstallBundleUtil::GetInstallBundleSubsetSectionPrefix() + BundleSubsets[iSubset];
 						InstallBundleConfig->GetArray(*SubsetSection, TEXT("FileRegex"), StrSubsetSearchRegexPatterns);
 					}
-
+					
 					TArray<FRegexPattern> SearchRegexPatterns;
 					SearchRegexPatterns.Reserve(StrSearchRegexPatterns.Num());
 					for (const FString& Str : StrSearchRegexPatterns)
@@ -434,10 +487,16 @@ void FInstallBundleSourceBulk::AsyncInit_MakeBundlesForBulkBuild()
 					{
 						SubsetSearchRegexPatterns.Emplace(Str, ERegexPatternFlags::CaseInsensitive);
 					}
-				
-					for (int i = 0; i < FoundFiles->Num();)
+					
+					for (int32 I = 0; I < FoundFiles->Num(); ++I)
 					{
-						FString& File = (*FoundFiles)[i];
+						// Skip if file was already matched by a bundle with higher order
+						if (MatchedBy[I].load(std::memory_order_relaxed) <= SectionIndex)
+						{
+							continue;
+						}
+						
+						FString& File = (*FoundFiles)[I];
 
 						bool bMatches = false;
 						for (const FRegexPattern& Pattern : SearchRegexPatterns)
@@ -451,32 +510,55 @@ void FInstallBundleSourceBulk::AsyncInit_MakeBundlesForBulkBuild()
 
 						if (!bMatches)
 						{
-							++i;
 							continue;
 						}
 
 						LOG_SOURCE_BULK(Verbose, TEXT("Adding %s to Bundle %s"), *File, *BundleName);
 
-						FString MatchedFile = MoveTemp(File);
-						FoundFiles->RemoveAtSwap(i);
-
 						bMatches = SubsetSearchRegexPatterns.IsEmpty();
 						for (const FRegexPattern& Pattern : SubsetSearchRegexPatterns)
 						{
-							if (FRegexMatcher(Pattern, MatchedFile).FindNext())
+							if (FRegexMatcher(Pattern, File).FindNext())
 							{
 								bMatches = true;
 								break;
 							}
 						}
 						
-						if (bMatches)
+						// This should really be a locked min - which we don't have.
+						// Would have to be a CAS loop but from my testing that was slower than doing this.
+						// This will give us some false positives but it's "good enough" and saves quite a bit of time.
+						if (MatchedBy[I].load(std::memory_order_relaxed) > SectionIndex)
+						{
+							MatchedBy[I].store(SectionIndex, std::memory_order_relaxed);
+						}
+
+						Data.MatchedFiles[I] = bMatches;
+						Data.bAnyMatches |= bMatches;
+						if (!bMatches)
 						{
-							BundleFileList.AddUnique(MoveTemp(MatchedFile));
+							LOG_SOURCE_BULK(Verbose, TEXT("Excluded %s from Bundle %s with subset %s"), *File, *BundleName, *BundleSubsets[iSubset]);
 						}
-						else
+					}
+				}, ParallelFlags);
+				
+				// We want files to be added to the BulkBuildBundles in the order that they are defined in SectionNames
+				TBitArray<> FoundFilesBits;
+				FoundFilesBits.Init(false, FoundFiles->Num());
+				for (int32 SectionIndex = 0; SectionIndex < SectionNames.Num(); ++SectionIndex)
+				{
+					FSectionParallelData& Data = ParallelData[SectionIndex];
+					if (Data.bAnyMatches)
+					{
+						TArray<FString>& BundleFileList = BulkBuildBundles.FindChecked(FName(*Data.BundleName));
+						for (int32 I = 0; I < FoundFiles->Num(); ++I)
 						{
-							LOG_SOURCE_BULK(Verbose, TEXT("Excluded %s from Bundle %s with subset %s"), *MatchedFile, *BundleName, *BundleSubsets[iSubset]);
+							FString& File = (*FoundFiles)[I];
+							if (Data.MatchedFiles[I] && !FoundFilesBits[I])
+							{
+								FoundFilesBits[I] = true;
+								BundleFileList.AddUnique(MoveTemp(File));
+							}
 						}
 					}
 				}
