Allow an IInstallBundleSource implementation to override the default PAK mount priority via its FInstallBundleSourceUpdateContentResultInfo

#Synced-CL 50324810

diff --git a/Engine/Plugins/Experimental/DefaultInstallBundleManager/Source/Private/DefaultInstallBundleManager.cpp b/Engine/Plugins/Experimental/DefaultInstallBundleManager/Source/Private/DefaultInstallBundleManager.cpp
--- a/Engine/Plugins/Experimental/DefaultInstallBundleManager/Source/Private/DefaultInstallBundleManager.cpp
+++ b/Engine/Plugins/Experimental/DefaultInstallBundleManager/Source/Private/DefaultInstallBundleManager.cpp
@@ -1384,6 +1384,7 @@ void FDefaultInstallBundleManager::UpdateBundleSourceComplete(TSharedRef<IInstal
 
 			BundleInfo.ContentPaths.AdditionalRootDirs.Append(MoveTemp(ResultInfo.AdditionalRootDirs));
 			BundleInfo.ContentPaths.ProjectName = ResultInfo.ProjectName;
+			BundleInfo.ContentPaths.PakOrder = FMath::Max(BundleInfo.ContentPaths.PakOrder, ResultInfo.PakOrder);
 		}
 
 		if (ResultInfo.bContentWasInstalled)
@@ -1658,9 +1659,10 @@ void FDefaultInstallBundleManager::MountPaks(FContentRequestRef Request)
 
 	auto MountPaksWorkFunc =
 		[ContentPaths = GetPakMountList(Request, BundleInfo)
+		, PakOrder = BundleInfo.ContentPaths.PakOrder
 		, LogVerbosityOverride = Request->LogVerbosityOverride](TValueOrError<void, UE::UnifiedError::FError>& OutResult) mutable -> bool
 	{
-		return MountPaksInList(ContentPaths, OutResult, LogVerbosityOverride);
+		return MountPaksInList(ContentPaths, OutResult, LogVerbosityOverride, PakOrder);
 	};
 
 	auto OnPakMountCompleteFunc = [this, WeakThis = AsWeak(), Request](bool bMountedPaks, const TValueOrError<void, UE::UnifiedError::FError>& MountResult) mutable
@@ -1874,7 +1876,7 @@ void FDefaultInstallBundleManager::MountPaks(FContentRequestRef Request)
 	}
 }
 
-bool FDefaultInstallBundleManager::MountPaksInList(TArrayView<TPair<FString, FPakMountOptions>> Paths, TValueOrError<void, UE::UnifiedError::FError>& OutResult, ELogVerbosity::Type LogVerbosityOverride)
+bool FDefaultInstallBundleManager::MountPaksInList(TArrayView<TPair<FString, FPakMountOptions>> Paths, TValueOrError<void, UE::UnifiedError::FError>& OutResult, ELogVerbosity::Type LogVerbosityOverride, const int32 PakOrder)
 {
 	OutResult = MakeValue();
 
@@ -1885,7 +1887,7 @@ bool FDefaultInstallBundleManager::MountPaksInList(TArrayView<TPair<FString, FPa
 
 	if (!FCoreDelegates::MountPaksEx.IsBound())
 	{
-		ensureMsgf(false, TEXT("Pak files have not been correctly initalized. Use -UsePaks on the cmdline if you are using the UnrealEditor.exe"));
+		ensureMsgf(false, TEXT("Pak files have not been correctly initialized. Use -UsePaks on the cmdline if you are using the UnrealEditor.exe"));
 		return false;
 	}
 
@@ -1897,6 +1899,7 @@ bool FDefaultInstallBundleManager::MountPaksInList(TArrayView<TPair<FString, FPa
 		MountArgs.Emplace(UE::FMountPaksExArgs
 		{
 			.PakFilePath = *Pair.Key,
+			.Order = PakOrder,
 			.MountOptions = Pair.Value
 		});
 	}
diff --git a/Engine/Plugins/Experimental/DefaultInstallBundleManager/Source/Public/DefaultInstallBundleManager.h b/Engine/Plugins/Experimental/DefaultInstallBundleManager/Source/Public/DefaultInstallBundleManager.h
--- a/Engine/Plugins/Experimental/DefaultInstallBundleManager/Source/Public/DefaultInstallBundleManager.h
+++ b/Engine/Plugins/Experimental/DefaultInstallBundleManager/Source/Public/DefaultInstallBundleManager.h
@@ -156,6 +156,7 @@ protected:
 		TArray<TPair<FString, FPakMountOptions>> ContentPaths;
 		TArray<FString> AdditionalRootDirs;
 		FString ProjectName;
+		int32 PakOrder = INDEX_NONE;
 		bool bContainsChunks = false;
 	};
 
@@ -568,7 +569,7 @@ protected:
 
 	UE_API void MountPaks(FContentRequestRef Request);
 
-	static UE_API bool MountPaksInList(TArrayView<TPair<FString, FPakMountOptions>> Paths, TValueOrError<void, UE::UnifiedError::FError>& OutResult, ELogVerbosity::Type LogVerbosityOverride);
+	static UE_API bool MountPaksInList(TArrayView<TPair<FString, FPakMountOptions>> Paths, TValueOrError<void, UE::UnifiedError::FError>& OutResult, ELogVerbosity::Type LogVerbosityOverride, const int32 PakOrder = INDEX_NONE);
 
 	UE_API void UnmountPaks(FContentReleaseRequestRef Request);
 
diff --git a/Engine/Source/Runtime/InstallBundleManager/Public/InstallBundleTypes.h b/Engine/Source/Runtime/InstallBundleManager/Public/InstallBundleTypes.h
--- a/Engine/Source/Runtime/InstallBundleManager/Public/InstallBundleTypes.h
+++ b/Engine/Source/Runtime/InstallBundleManager/Public/InstallBundleTypes.h
@@ -408,6 +408,7 @@ struct FInstallBundleSourceUpdateContentResultInfo
 	TSet<FString> NonUFSShaderLibPaths;
 	TArray<TUniquePtr<UE::IoStore::FOnDemandMountArgs>> OnDemandMountArgs;
 	FPakMountOptions MountOptions;
+	int32 PakOrder = INDEX_NONE;
 	FString ProjectName;
 
 	uint64 CurrentInstallSize = 0;
