Deprecate InstallBundleManagerUtil::MakePlatformBundleSource and add InstallBundleManagerUtil::GetPlatformChunkInstall instead.
The "Platform" bundle source  is now always the chunk install wrapper.
This makes it easier for games to subclass FInstallBundleSourcePlatformChunkInstall

#Synced-CL 50292886

diff --git a/Engine/Plugins/Experimental/DefaultInstallBundleManager/Source/Private/InstallBundleManagerUtil.cpp b/Engine/Plugins/Experimental/DefaultInstallBundleManager/Source/Private/InstallBundleManagerUtil.cpp
--- a/Engine/Plugins/Experimental/DefaultInstallBundleManager/Source/Private/InstallBundleManagerUtil.cpp
+++ b/Engine/Plugins/Experimental/DefaultInstallBundleManager/Source/Private/InstallBundleManagerUtil.cpp
@@ -4,6 +4,7 @@
 
 #include "DefaultInstallBundleManagerPrivate.h"
 #include "InstallBundleManagerInterface.h"
+#include "InstallBundleSourcePlatformChunkInstall.h"
 #include "Misc/ConfigCacheIni.h"
 
 #include "InstallBundleSourceBulk.h"
@@ -38,9 +39,10 @@ namespace InstallBundleManagerUtil
 #if WITH_PLATFORM_INSTALL_BUNDLE_SOURCE
 		if (Type.GetName() == TEXTVIEW("Platform"))
 		{
-			if (TSharedPtr<IInstallBundleSource> PlatformSource = MakePlatformBundleSource())
+			IPlatformChunkInstall* PlatformChunkInstall = InstallBundleManagerUtil::GetPlatformChunkInstall();
+			if (PlatformChunkInstall)
 			{
-				return PlatformSource;
+				return MakeShared<FInstallBundleSourcePlatformChunkInstall>(PlatformChunkInstall);
 			}
 		}
 #endif // WITH_PLATFORM_INSTALL_BUNDLE_SOURCE
diff --git a/Engine/Plugins/Experimental/DefaultInstallBundleManager/Source/Public/InstallBundleManagerUtil.h b/Engine/Plugins/Experimental/DefaultInstallBundleManager/Source/Public/InstallBundleManagerUtil.h
--- a/Engine/Plugins/Experimental/DefaultInstallBundleManager/Source/Public/InstallBundleManagerUtil.h
+++ b/Engine/Plugins/Experimental/DefaultInstallBundleManager/Source/Public/InstallBundleManagerUtil.h
@@ -24,6 +24,9 @@ namespace InstallBundleManagerUtil
 	DEFAULTINSTALLBUNDLEMANAGER_API TSharedPtr<IInstallBundleSource> MakeBundleSource(FInstallBundleSourceType Type);
 
 #if WITH_PLATFORM_INSTALL_BUNDLE_SOURCE
+	DEFAULTINSTALLBUNDLEMANAGER_API IPlatformChunkInstall* GetPlatformChunkInstall();
+
+	UE_DEPRECATED(5.8, "Implement InstallBundleManagerUtil::GetPlatformChunkInstall instead.")
 	DEFAULTINSTALLBUNDLEMANAGER_API TSharedPtr<IInstallBundleSource> MakePlatformBundleSource();
 #endif
 
