Add new analytics that track delays in async startup and exit.  Adds a flag (and total count) for installs that the app went to background during.

#rnx

diff --git a/Engine/Plugins/Experimental/DefaultInstallBundleManager/Source/Private/InstallBundleManagerUtil.cpp b/Engine/Plugins/Experimental/DefaultInstallBundleManager/Source/Private/InstallBundleManagerUtil.cpp
--- a/Engine/Plugins/Experimental/DefaultInstallBundleManager/Source/Private/InstallBundleManagerUtil.cpp
+++ b/Engine/Plugins/Experimental/DefaultInstallBundleManager/Source/Private/InstallBundleManagerUtil.cpp
@@ -833,7 +833,11 @@ namespace InstallBundleManagerAnalytics
 			Attributes.Add(FAnalyticsEventAttribute(TEXT("TotalDecompressionDecompressTime"), RequestCompleteInfo.BundleDecompressionInfo.TotalDecompressionDecompressTime));
 			Attributes.Add(FAnalyticsEventAttribute(TEXT("TotalDecompressionWriteTime"), RequestCompleteInfo.BundleDecompressionInfo.TotalDecompressionWriteTime));
 		}
-		
+
+		Attributes.Add(FAnalyticsEventAttribute(TEXT("AsyncStartupDelay"), RequestCompleteInfo.AsyncStartupDelay));
+		Attributes.Add(FAnalyticsEventAttribute(TEXT("AsyncExitDelay"), RequestCompleteInfo.AsyncExitDelay));
+		Attributes.Add(FAnalyticsEventAttribute(TEXT("AsyncPlacedInBackground"), RequestCompleteInfo.AsyncPlacedInBackground));
+
 		FString ChunckDonwloadInfoString;
 		FInstallBundleManagerUtilJsonWriter JsonWriter(&ChunckDonwloadInfoString);
 		JsonWriter.WriteArrayStart();
diff --git a/Engine/Plugins/Experimental/DefaultInstallBundleManager/Source/Public/InstallBundleManagerUtil.h b/Engine/Plugins/Experimental/DefaultInstallBundleManager/Source/Public/InstallBundleManagerUtil.h
--- a/Engine/Plugins/Experimental/DefaultInstallBundleManager/Source/Public/InstallBundleManagerUtil.h
+++ b/Engine/Plugins/Experimental/DefaultInstallBundleManager/Source/Public/InstallBundleManagerUtil.h
@@ -373,6 +373,9 @@ namespace InstallBundleManagerAnalytics
 	 * @EventParam EstimatedFullDownloadBytes (uint64) Estimated bytes that needed download
 	 * @EventParam BPTErrorCode (string) Any BPT Error
 	 * @EventParam DownloadType (string) Type of download
+	 * @EventParam AsyncStartupDelay (double) How long (in seconds) before the async install task starts
+	 * @EventParam AsyncExitDelay (double) How long (in seconds) before the async OnComplete is processed
+	 * @EventParam AsyncPlacedInBackground (bool) True if app was placed in background during install step
 	 * @EventParam ChunkDownloadCompleteInfos (string) Detail metrics of each chunk downloaded, if any
 	*/
 	DEFAULTINSTALLBUNDLEMANAGER_API void FireEvent_BundleRequestCompleteDetailed(IAnalyticsProviderET* AnalyticsProvider,
diff --git a/Engine/Source/Runtime/InstallBundleManager/Public/InstallBundleTypes.h b/Engine/Source/Runtime/InstallBundleManager/Public/InstallBundleTypes.h
--- a/Engine/Source/Runtime/InstallBundleManager/Public/InstallBundleTypes.h
+++ b/Engine/Source/Runtime/InstallBundleManager/Public/InstallBundleTypes.h
@@ -529,4 +529,7 @@ struct FBundleRequestCompleteInfo
 	FString BPTErrorCode;
 	FString DownloadType;
 	FBundleDecompressionInfo BundleDecompressionInfo;
+	double AsyncStartupDelay{};
+	double AsyncExitDelay{};
+	bool AsyncPlacedInBackground{};
 };
