More IAX mount args cleanup
- Remove EOnDemandMountOptions that no longer do anything.
- If a bundle source does not pass back OnDemandMountArgs, bundle manager will now create them.

#Synced-CL 50389196

diff --git a/Engine/Plugins/Experimental/DefaultInstallBundleManager/Source/Private/DefaultInstallBundleManager.cpp b/Engine/Plugins/Experimental/DefaultInstallBundleManager/Source/Private/DefaultInstallBundleManager.cpp
--- a/Engine/Plugins/Experimental/DefaultInstallBundleManager/Source/Private/DefaultInstallBundleManager.cpp
+++ b/Engine/Plugins/Experimental/DefaultInstallBundleManager/Source/Private/DefaultInstallBundleManager.cpp
@@ -1305,7 +1305,9 @@ void FDefaultInstallBundleManager::UpdateBundleSourceComplete(TSharedRef<IInstal
 	}
 
 	if (Request->SourceRequestResults.Num() != Request->RequiredSourceRequestResultsCount)
+	{
 		return;
+	}
 
 	LOG_INSTALL_BUNDLE_MAN_OVERRIDE(Request->LogVerbosityOverride, Display, TEXT("Bundle %s done waiting for all bundle sources!"), *Request->BundleName.ToString());
 
@@ -1370,6 +1372,7 @@ void FDefaultInstallBundleManager::UpdateBundleSourceComplete(TSharedRef<IInstal
 		}
 		else if (Request->Result == EInstallBundleResult::OK)
 		{
+			TArray<TUniquePtr<UE::IoStore::FOnDemandMountArgs>> CreatedOnDemandMountArgs;
 			for (FString& ContentPath : ResultInfo.ContentPaths)
 			{
 				// We only need to keep track of the pak paths for unmounting so 
@@ -1378,8 +1381,25 @@ void FDefaultInstallBundleManager::UpdateBundleSourceComplete(TSharedRef<IInstal
 				{
 					BundleInfo.ContentPaths.ContentPaths.Emplace(MoveTemp(ContentPath), ResultInfo.MountOptions);
 				}
+
+				// If the bundle source did not supply OnDemandMountArgs, create them now.
+				// This assumes that the bundle source has already correctly setup ondemand hostgroups.
+				if (ResultInfo.OnDemandMountArgs.IsEmpty() && ContentPath.EndsWith(TEXTVIEW(".uondemandtoc")))
+				{
+					TUniquePtr<UE::IoStore::FOnDemandMountArgs> MountArgs = MakeUnique<UE::IoStore::FOnDemandMountArgs>();
+					MountArgs->MountId = Request->BundleName.ToString();
+					MountArgs->FilePath = ContentPath;
+
+					if (EnumHasAnyFlags(ResultInfo.MountOptions.MountFlags, FPakMountOptions::EMountFlags::WithSoftReferences))
+					{
+						MountArgs->Options |= UE::IoStore::EOnDemandMountOptions::WithSoftReferences;
+					}
+
+					CreatedOnDemandMountArgs.Add(MoveTemp(MountArgs));
+				}
 			}
 
+			Request->OnDemandMountArgs.Append(MoveTemp(CreatedOnDemandMountArgs));
 			Request->OnDemandMountArgs.Append(MoveTemp(ResultInfo.OnDemandMountArgs));
 
 			BundleInfo.ContentPaths.AdditionalRootDirs.Append(MoveTemp(ResultInfo.AdditionalRootDirs));
diff --git a/Engine/Plugins/Experimental/DefaultInstallBundleManager/Source/Private/InstallBundleSourceBulk.cpp b/Engine/Plugins/Experimental/DefaultInstallBundleManager/Source/Private/InstallBundleSourceBulk.cpp
--- a/Engine/Plugins/Experimental/DefaultInstallBundleManager/Source/Private/InstallBundleSourceBulk.cpp
+++ b/Engine/Plugins/Experimental/DefaultInstallBundleManager/Source/Private/InstallBundleSourceBulk.cpp
@@ -6,7 +6,6 @@
 #include "DefaultInstallBundleManagerPrivate.h"
 #include "HAL/PlatformFileManager.h"
 #include "HAL/PlatformFile.h"
-#include "IO/IoStoreOnDemand.h"
 #include "IPlatformFilePak.h"
 #include "InstallBundleManagerInterface.h"
 #include "InstallBundleManagerUtil.h"
@@ -743,23 +742,6 @@ void FInstallBundleSourceBulk::RequestUpdateContent(FRequestUpdateContentBundleC
 
 	FilterContentBundlePaths(Context, ResultInfo.ContentPaths);
 
-	for (const FString& Path : ResultInfo.ContentPaths)
-	{
-		if (Path.EndsWith(TEXTVIEW(".uondemandtoc")))
-		{
-			TUniquePtr<UE::IoStore::FOnDemandMountArgs> MountArgs = MakeUnique<UE::IoStore::FOnDemandMountArgs>();
-			MountArgs->MountId = Context.BundleName.ToString();
-			MountArgs->FilePath = Path;
-
-			if (OnDemandMountOptions.bWithSoftReferences)
-			{
-				MountArgs->Options |= UE::IoStore::EOnDemandMountOptions::WithSoftReferences;
-			}
-
-			ResultInfo.OnDemandMountArgs.Add(MoveTemp(MountArgs));
-		}
-	}
-
 	if (OnDemandMountOptions.bWithSoftReferences)
 	{
 		ResultInfo.MountOptions.MountFlags |= FPakMountOptions::EMountFlags::WithSoftReferences;
diff --git a/Engine/Plugins/Experimental/DefaultInstallBundleManager/Source/Private/InstallBundleSourcePlatformChunkInstall.cpp b/Engine/Plugins/Experimental/DefaultInstallBundleManager/Source/Private/InstallBundleSourcePlatformChunkInstall.cpp
--- a/Engine/Plugins/Experimental/DefaultInstallBundleManager/Source/Private/InstallBundleSourcePlatformChunkInstall.cpp
+++ b/Engine/Plugins/Experimental/DefaultInstallBundleManager/Source/Private/InstallBundleSourcePlatformChunkInstall.cpp
@@ -12,7 +12,6 @@
 #include "Algo/AnyOf.h"
 #include "Algo/Find.h"
 #include "Containers/Ticker.h"
-#include "IO/IoStoreOnDemand.h"
 #include "IPlatformFilePak.h"
 #include "Stats/Stats.h"
 
@@ -542,30 +541,6 @@ void FInstallBundleSourcePlatformChunkInstall::RequestUpdateContent(FRequestUpda
 			ResultInfo.MountOptions.MountFlags |= FPakMountOptions::EMountFlags::WithSoftReferences;
 		}
 
-		// TODO: this no longer needs to happen in the bundle source
-		// Move this into IBM
-		for (const FString& Path : ResultInfo.ContentPaths)
-		{
-			if (Path.EndsWith(TEXTVIEW(".uondemandtoc")))
-			{
-				TUniquePtr<UE::IoStore::FOnDemandMountArgs> MountArgs = MakeUnique<UE::IoStore::FOnDemandMountArgs>();
-				MountArgs->MountId = Context.BundleName.ToString();
-				//GetOnDemandHostGroup(MountArgs->HostGroup);
-				//MountArgs->TocRelativeUrl = GetOnDemandTocRelativeURL();
-				MountArgs->FilePath = Path;
-				MountArgs->Options = UE::IoStore::EOnDemandMountOptions::InstallOnDemand;
-				// IAX content on the BPS CDN's does not require https
-				MountArgs->Options |= UE::IoStore::EOnDemandMountOptions::ForceNonSecureHttp;
-
-				if (ConfigMountOptions.bWithSoftReferences)
-				{
-					MountArgs->Options |= UE::IoStore::EOnDemandMountOptions::WithSoftReferences;
-				}
-
-				ResultInfo.OnDemandMountArgs.Add(MoveTemp(MountArgs));
-			}
-		}
-
 		Context.CompleteCallback.ExecuteIfBound(AsShared(), MoveTemp(ResultInfo));
 	}
 	else
@@ -739,30 +714,6 @@ void FInstallBundleSourcePlatformChunkInstall::OnNamedChunkInstall( const FNamed
 				ResultInfo.MountOptions.MountFlags |= FPakMountOptions::EMountFlags::WithSoftReferences;
 			}
 
-			// TODO: this no longer needs to happen in the bundle source
-			// Move this into IBM
-			for (const FString& Path : ResultInfo.ContentPaths)
-			{
-				if (Path.EndsWith(TEXTVIEW(".uondemandtoc")))
-				{
-					TUniquePtr<UE::IoStore::FOnDemandMountArgs> MountArgs = MakeUnique<UE::IoStore::FOnDemandMountArgs>();
-					MountArgs->MountId = Request->BundleName.ToString();
-					//GetOnDemandHostGroup(MountArgs->HostGroup);
-					//MountArgs->TocRelativeUrl = GetOnDemandTocRelativeURL();
-					MountArgs->FilePath = Path;
-					MountArgs->Options = UE::IoStore::EOnDemandMountOptions::InstallOnDemand;
-					// IAX content on the BPS CDN's does not require https
-					MountArgs->Options |= UE::IoStore::EOnDemandMountOptions::ForceNonSecureHttp;
-
-					if (ConfigMountOptions.bWithSoftReferences)
-					{
-						MountArgs->Options |= UE::IoStore::EOnDemandMountOptions::WithSoftReferences;
-					}
-
-					ResultInfo.OnDemandMountArgs.Add(MoveTemp(MountArgs));
-				}
-			}
-
 			Request->CompleteCallback.ExecuteIfBound(AsShared(), MoveTemp(ResultInfo));
 		}
 
diff --git a/Engine/Source/Runtime/Experimental/IoStore/OnDemand/Private/OnDemandIoStore.cpp b/Engine/Source/Runtime/Experimental/IoStore/OnDemand/Private/OnDemandIoStore.cpp
--- a/Engine/Source/Runtime/Experimental/IoStore/OnDemand/Private/OnDemandIoStore.cpp
+++ b/Engine/Source/Runtime/Experimental/IoStore/OnDemand/Private/OnDemandIoStore.cpp
@@ -480,8 +480,7 @@ FIoStatus FOnDemandIoStore::InitializePostHotfix()
 			{
 				.MountId = TEXT("global"),
 				.FilePath = Config.TocFilePath,
-				.HostGroup = HostGroup,
-				.Options = EOnDemandMountOptions::StreamOnDemand
+				.HostGroup = HostGroup
 			});
 		}
 	}
@@ -500,7 +499,6 @@ FIoStatus FOnDemandIoStore::InitializePostHotfix()
 		if (MountArgs)
 		{
 			MountArgs.GetValue().HostGroup = HostGroup.ConsumeValueOrDie(); 
-			MountArgs.GetValue().Options = EOnDemandMountOptions::InstallOnDemand;
 			MountArgs.GetValue().TocRelativeUrl = Config.TocPath;
 
 			static FOnDemandContentHandle ContentHandle = FOnDemandContentHandle::Create(FStringView(TEXT("AllContent")));
@@ -1574,23 +1572,6 @@ FIoStatus FOnDemandIoStore::CreateContainersFromToc(
 
 FIoStatus FOnDemandIoStore::SetupHostGroup(const FSharedOnDemandContainer& Container, const FOnDemandMountArgs& MountArgs)
 {
-	auto SanitizeHostNames = [MountOptions = MountArgs.Options](TConstArrayView<FAnsiString> HostNames) -> TArray<FAnsiString>
-	{
-		TArray<FAnsiString> OutHostNames;
-		for (const FAnsiString& HostName : HostNames)
-		{
-			FAnsiString Tmp = HostName;
-			Tmp.ToLowerInline();
-			if (EnumHasAnyFlags(MountOptions, EOnDemandMountOptions::ForceNonSecureHttp))
-			{
-				Tmp.ReplaceInline("https", "http");
-			}
-			OutHostNames.Add(Tmp);
-		}
-
-		return OutHostNames;
-	};
-
 	if (Container->HostGroupName.IsNone())
 	{
 		Container->HostGroupName = MountArgs.HostGroupName.IsNone()
@@ -1611,8 +1592,7 @@ FIoStatus FOnDemandIoStore::SetupHostGroup(const FSharedOnDemandContainer& Conta
 
 	if (MountArgs.HostGroup.IsEmpty() == false)
 	{
-		TArray<FAnsiString> HostNames = SanitizeHostNames(MountArgs.HostGroup.Hosts());
-		return FHttpIoDispatcher::RegisterHostGroup(Container->HostGroupName, HostNames, Container->GetTestUrl());
+		return FHttpIoDispatcher::RegisterHostGroup(Container->HostGroupName, MountArgs.HostGroup.Hosts(), Container->GetTestUrl());
 	}
 
 	return FIoStatus(EIoErrorCode::PendingHostGroup); 
diff --git a/Engine/Source/Runtime/Experimental/IoStore/OnDemandCore/Private/IoStoreOnDemand.cpp b/Engine/Source/Runtime/Experimental/IoStore/OnDemandCore/Private/IoStoreOnDemand.cpp
--- a/Engine/Source/Runtime/Experimental/IoStore/OnDemandCore/Private/IoStoreOnDemand.cpp
+++ b/Engine/Source/Runtime/Experimental/IoStore/OnDemandCore/Private/IoStoreOnDemand.cpp
@@ -174,7 +174,7 @@ static void MountUrl(const TArray<FString>& Args)
 		return;
 	}
 
-	if (Args.Num() < 2)
+	if (Args.Num() < 1)
 	{
 		UE_LOG(LogIoStoreOnDemand, Error, TEXT("Not enough arguments."));
 		return;
@@ -191,21 +191,14 @@ static void MountUrl(const TArray<FString>& Args)
 		return;
 	}
 
-	EOnDemandMountOptions Options = EOnDemandMountOptions::StreamOnDemand;
-	const FString& Arg = Args[1];
-	if (Arg.ToLower().Contains(TEXT("install")))
-	{
-		Options = EOnDemandMountOptions::InstallOnDemand; 
-	}
-	const FString& MountId = Args.Num() > 2 ? Args[2] : Url; 
+	const FString& MountId = Args.Num() > 1 ? Args[1] : Url; 
 
 	IoStore->Mount(
 		FOnDemandMountArgs
 		{
 			.MountId = MountId,
 			.TocRelativeUrl = FString(TocRelativeUrl),
-			.HostGroup = HostGroup.ConsumeValueOrDie(),
-			.Options = Options
+			.HostGroup = HostGroup.ConsumeValueOrDie()
 		},
 		[](FOnDemandMountResult MountResult)
 		{
@@ -223,7 +216,7 @@ static void MountFile(const TArray<FString>& Args)
 		return;
 	}
 
-	if (Args.Num() < 3)
+	if (Args.Num() < 2)
 	{
 		UE_LOG(LogIoStoreOnDemand, Error, TEXT("Not enough arguments."));
 		return;
@@ -261,13 +254,7 @@ static void MountFile(const TArray<FString>& Args)
 		return;
 	}
 
-	EOnDemandMountOptions Options = EOnDemandMountOptions::StreamOnDemand;
-	if (Args[1].ToLower().Contains(TEXT("install")))
-	{
-		Options = EOnDemandMountOptions::InstallOnDemand; 
-	}
-
-	FString HostPath = Args[2];
+	FString HostPath = Args[1];
 	HostPath.TrimQuotesInline();
 	FStringView TocRelativeUrl;
 	TIoStatusOr<FOnDemandHostGroup> HostResult = SplitHostAndTocUrl(HostPath, TocRelativeUrl);
@@ -279,7 +266,7 @@ static void MountFile(const TArray<FString>& Args)
 
 	FOnDemandHostGroup HostGroup = HostResult.ConsumeValueOrDie();
 
-	const FString& MountId = Args.Num() > 3 ? Args[3] : FilePaths[0];
+	const FString& MountId = Args.Num() > 2 ? Args[2] : FilePaths[0];
 	for (const FString& FilePath : FilePaths)
 	{
 		FOnDemandMountArgs& MountArgs	= AllMountArgs.AddDefaulted_GetRef();
@@ -297,7 +284,6 @@ static void MountFile(const TArray<FString>& Args)
 
 	for (FOnDemandMountArgs& MountArgs : AllMountArgs)
 	{
-		MountArgs.Options = Options;
 		IoStore->Mount(
 			MoveTemp(MountArgs),
 			[](FOnDemandMountResult MountResult)
diff --git a/Engine/Source/Runtime/Experimental/IoStore/OnDemandCore/Public/IO/IoStoreOnDemand.h b/Engine/Source/Runtime/Experimental/IoStore/OnDemandCore/Public/IO/IoStoreOnDemand.h
--- a/Engine/Source/Runtime/Experimental/IoStore/OnDemandCore/Public/IO/IoStoreOnDemand.h
+++ b/Engine/Source/Runtime/Experimental/IoStore/OnDemandCore/Public/IO/IoStoreOnDemand.h
@@ -149,16 +149,12 @@ struct FOnDemandRegisterHostGroupResult
 /** Options for controlling the behavior of mounted container(s). */
 enum class EOnDemandMountOptions
 {
-	/** Mount containers with the purpose of streaming the content on-demand. */
-	StreamOnDemand			= 1 << 0,
-	/** Mount containers with the purpose of installing/downloading the content on-demand. */
-	InstallOnDemand			= 1 << 1,
+	/** No additional options. */
+	None					= 0,
 	/** Trigger callback on game thread. */
 	CallbackOnGameThread	= 1 << 2,
 	/** Make soft references available */
-	WithSoftReferences		= 1 << 3,
-	/** Force the use of non-secure HTTP protocol. */
-	ForceNonSecureHttp		= 1 << 4,
+	WithSoftReferences		= 1 << 3
 };
 ENUM_CLASS_FLAGS(EOnDemandMountOptions);
 
@@ -193,7 +189,7 @@ struct FOnDemandMountArgs
 	/** List of URLs from where to download the chunks. */
 	FOnDemandHostGroup HostGroup;
 	/** Mount options. */
-	EOnDemandMountOptions Options = EOnDemandMountOptions::StreamOnDemand;
+	EOnDemandMountOptions Options = EOnDemandMountOptions::None;
 };
 
 /** Holds information about a mount request. */
