Modified BundleSourcePlatformChunkInstall to always return download size

#rnx
#Synced-CL 51877528

diff --git a/Engine/Plugins/Experimental/DefaultInstallBundleManager/Source/Private/InstallBundleSourcePlatformChunkInstall.cpp b/Engine/Plugins/Experimental/DefaultInstallBundleManager/Source/Private/InstallBundleSourcePlatformChunkInstall.cpp
--- a/Engine/Plugins/Experimental/DefaultInstallBundleManager/Source/Private/InstallBundleSourcePlatformChunkInstall.cpp
+++ b/Engine/Plugins/Experimental/DefaultInstallBundleManager/Source/Private/InstallBundleSourcePlatformChunkInstall.cpp
@@ -421,7 +421,6 @@ void FInstallBundleSourcePlatformChunkInstall::GetContentState(TArrayView<const
 	State.CurrentVersion.Add(GetSourceType(), InstallBundleUtil::GetAppVersion());
 
 	// capture the install state and version for all bundles, and record the remaining download size
-	uint64 TotalRemainingDownloadSize = 0;
 	TMap<FName,uint64> BundleToRemaingDownloadSize;
 	for (const FName& BundleName : BundleNames)
 	{
@@ -450,7 +449,10 @@ void FInstallBundleSourcePlatformChunkInstall::GetContentState(TArrayView<const
 			{
 				check( ChunkStatusDetail.CurrentInstallSize <= ChunkStatusDetail.FullInstallSize);
 				RemainingDownloadSize = (ChunkStatusDetail.FullInstallSize - ChunkStatusDetail.CurrentInstallSize);
-				TotalRemainingDownloadSize += RemainingDownloadSize;
+
+				State.ContentSize.DownloadSize += RemainingDownloadSize;
+				State.ContentSize.CurrentSizeOnDiskOtherDirs += ChunkStatusDetail.CurrentInstallSize;
+				State.ContentSize.SpaceRequiredForInstallOtherDirs += RemainingDownloadSize;
 			}
 			BundleToRemaingDownloadSize.Add(BundleName, RemainingDownloadSize);
 		}
@@ -459,17 +461,16 @@ void FInstallBundleSourcePlatformChunkInstall::GetContentState(TArrayView<const
 	// compute the download weight for all of the known bundles - higher weight is a bigger download
 	for (TTuple<FName,FInstallBundleContentState>& BundleStatePair : State.IndividualBundleStates)
 	{
-		if (TotalRemainingDownloadSize == 0)
+		if (State.ContentSize.DownloadSize == 0)
 		{
 			BundleStatePair.Value.Weight = 1.0f / BundleNames.Num();
 			continue;
 		}
 
-		double Weight = (double)BundleToRemaingDownloadSize[BundleStatePair.Key] / (double)TotalRemainingDownloadSize;
+		double Weight = (double)BundleToRemaingDownloadSize[BundleStatePair.Key] / (double)State.ContentSize.DownloadSize;
 		BundleStatePair.Value.Weight = FMath::Max(Weight, InstallBundleUtil::MinimumBundleWeight); // this Max() does mean the total weight will be > 1.0 but all other bundle sources do it this way too
 	}
 
-	// Don't return any size info since all the size should be reserved by the system	
 	Callback.ExecuteIfBound(MoveTemp(State));
 }
 
