Add todo comment

#Synced-CL 51330535

diff --git a/Engine/Plugins/Experimental/DefaultInstallBundleManager/Source/Private/DefaultInstallBundleManager.cpp b/Engine/Plugins/Experimental/DefaultInstallBundleManager/Source/Private/DefaultInstallBundleManager.cpp
--- a/Engine/Plugins/Experimental/DefaultInstallBundleManager/Source/Private/DefaultInstallBundleManager.cpp
+++ b/Engine/Plugins/Experimental/DefaultInstallBundleManager/Source/Private/DefaultInstallBundleManager.cpp
@@ -1714,6 +1714,7 @@ void FDefaultInstallBundleManager::MountPaks(FContentRequestRef Request)
 
 		for (const FString& RootDir : BundleInfoLambda.ContentPaths.AdditionalRootDirs)
 		{
+			// TODO: We should have a function to remove these so they don't leak at metaverse scale
 			FPlatformMisc::AddAdditionalRootDirectory(RootDir);
 		}
 
